DeviceManagement: Initiializing device name and platform name for target device proxies with placeholder values extracted from the DeviceID, so we have something to show to the user until the device is connected and a pong message has been received from it

diff --git a/Engine/Source/Developer/TargetDeviceServices/Private/TargetDeviceProxy.cpp b/Engine/Source/Developer/TargetDeviceServices/Private/TargetDeviceProxy.cpp
--- a/Engine/Source/Developer/TargetDeviceServices/Private/TargetDeviceProxy.cpp
+++ b/Engine/Source/Developer/TargetDeviceServices/Private/TargetDeviceProxy.cpp
@@ -19,6 +19,14 @@ FTargetDeviceProxy::FTargetDeviceProxy( const FString& InId )
 	, SupportsPowerOn(false)
 	, SupportsReboot(false)
 {
+	FTargetDeviceId DeviceId;
+
+	if (FTargetDeviceId::Parse(InId, DeviceId))
+	{
+		Name = DeviceId.GetDeviceName();
+		Platform = DeviceId.GetPlatformName();
+	}
+
 	InitializeMessaging();
 }
 
@@ -31,8 +39,8 @@ FTargetDeviceProxy::FTargetDeviceProxy( const FString& InId, const FTargetDevice
 	, SupportsPowerOn(false)
 	, SupportsReboot(false)
 {
-	InitializeMessaging();
 	UpdateFromMessage(Message, Context);
+	InitializeMessaging();
 }
 
 
