Docs: Removed file comments and added missing code documentation

diff --git a/Engine/Source/Developer/SlateReflector/Private/SlateReflectorModule.cpp b/Engine/Source/Developer/SlateReflector/Private/SlateReflectorModule.cpp
--- a/Engine/Source/Developer/SlateReflector/Private/SlateReflectorModule.cpp
+++ b/Engine/Source/Developer/SlateReflector/Private/SlateReflectorModule.cpp
@@ -1,9 +1,5 @@
 // Copyright 1998-2014 Epic Games, Inc. All Rights Reserved.
 
-/*=============================================================================
-	SlateReflectorModule.cpp: Implements the FSlateReflectorModule class.
-=============================================================================*/
-
 #include "SlateReflectorPrivatePCH.h"
 #include "ModuleManager.h"
 
@@ -19,7 +15,7 @@ class FSlateReflectorModule
 {
 public:
 
-	// Begin ISlateReflectorModule interface
+	// ISlateReflectorModule interface
 
 	virtual TSharedRef<SWidget> GetWidgetReflector( ) override
 	{
@@ -48,11 +44,9 @@ public:
 		FGlobalTabmanager::Get()->UnregisterNomadTabSpawner("WidgetReflector");
 	}
 
-	// End ISlateReflectorModule interface
-
 public:
 
-	// Begin IModuleInterface interface
+	// IModuleInterface interface
 
 	virtual void StartupModule( ) override { }
 
@@ -61,8 +55,6 @@ public:
 		UnregisterTabSpawner();
 	}
 
-	// End IModuleInterface interface
-
 private:
 
 	TSharedRef<SDockTab> MakeWidgetReflectorTab( const FSpawnTabArgs& )
diff --git a/Engine/Source/Developer/SlateReflector/Private/SlateReflectorPrivatePCH.h b/Engine/Source/Developer/SlateReflector/Private/SlateReflectorPrivatePCH.h
--- a/Engine/Source/Developer/SlateReflector/Private/SlateReflectorPrivatePCH.h
+++ b/Engine/Source/Developer/SlateReflector/Private/SlateReflectorPrivatePCH.h
@@ -1,9 +1,5 @@
 // Copyright 1998-2014 Epic Games, Inc. All Rights Reserved.
 
-/*=============================================================================
-	SlateReflectorPrivatePCH.h: Pre-compiled header file for the SlateReflector module.
-=============================================================================*/
-
 #pragma once
 
 #include "SlateReflector.h"
diff --git a/Engine/Source/Developer/SlateReflector/Private/Widgets/SWidgetReflector.cpp b/Engine/Source/Developer/SlateReflector/Private/Widgets/SWidgetReflector.cpp
--- a/Engine/Source/Developer/SlateReflector/Private/Widgets/SWidgetReflector.cpp
+++ b/Engine/Source/Developer/SlateReflector/Private/Widgets/SWidgetReflector.cpp
@@ -1,9 +1,5 @@
 // Copyright 1998-2014 Epic Games, Inc. All Rights Reserved.
 
-/*=============================================================================
-	SWidgetReflector.cpp: Implements the SWidgetReflector class.
-=============================================================================*/
-
 #include "SlateReflectorPrivatePCH.h"
 
 
diff --git a/Engine/Source/Developer/SlateReflector/Private/Widgets/SWidgetReflector.h b/Engine/Source/Developer/SlateReflector/Private/Widgets/SWidgetReflector.h
--- a/Engine/Source/Developer/SlateReflector/Private/Widgets/SWidgetReflector.h
+++ b/Engine/Source/Developer/SlateReflector/Private/Widgets/SWidgetReflector.h
@@ -1,9 +1,5 @@
 // Copyright 1998-2014 Epic Games, Inc. All Rights Reserved.
 
-/*=============================================================================
-	SWidgetReflector.h: Declares the SWidgetReflector class.
-=============================================================================*/
-
 #pragma once
 
 
@@ -26,7 +22,6 @@ public:
 	SLATE_BEGIN_ARGS(SWidgetReflector) { }
 	SLATE_END_ARGS()
 
-	/** Create widgets that comprise this WidgetReflector implementation */
 	/**
 	 * Creates and initializes a new widget reflector widget.
 	 *
@@ -36,15 +31,13 @@ public:
 
 public:
 
-	// Begin SCompoundWidget overrides
+	// SCompoundWidget overrides
 
 	virtual void Tick( const FGeometry& AllottedGeometry, const double InCurrentTime, const float InDeltaTime ) override;
 
-	// End SCompoundWidget overrides
-
 public:
 
-	// Begin IWidgetReflector interface
+	// IWidgetReflector interface
 
 	virtual bool IsInPickingMode( ) const override
 	{
@@ -74,18 +67,14 @@ public:
 	}
 
 	virtual void SetWidgetsToVisualize( const FWidgetPath& InWidgetsToVisualize ) override;
-
 	virtual int32 Visualize( const FWidgetPath& InWidgetsToVisualize, FSlateWindowElementList& OutDrawElements, int32 LayerId ) override;
 
-	// End IWidgetReflector interface
-
 protected:
 
 	/**
 	 * Generates a tool tip for the given reflector tree node.
 	 *
 	 * @param InReflectorNode The node to generate the tool tip for.
-	 *
 	 * @return The tool tip widget.
 	 */
 	TSharedRef<SToolTip> GenerateToolTipForReflectorNode( TSharedPtr<FReflectorNode> InReflectorNode );
@@ -103,7 +92,6 @@ protected:
 	 * @param InWidgetsToVisualize A widget path whose widgets' outlines to draw.
 	 * @param OutDrawElements A list of draw elements; we will add the output outlines into it.
 	 * @param LayerId The maximum layer achieved in OutDrawElements so far.
-	 *
 	 * @return The maximum layer ID we achieved while painting.
 	 */
 	int32 VisualizePickAsRectangles( const FWidgetPath& InWidgetsToVisualize, FSlateWindowElementList& OutDrawElements, int32 LayerId );
@@ -115,7 +103,6 @@ protected:
 	 * @param WindowGeometry The geometry of the window in which to draw.
 	 * @param OutDrawElements A list of draw elements; we will add the output outlines into it.
 	 * @param LayerId the maximum layer achieved in OutDrawElements so far.
-	 *
 	 * @return The maximum layer ID we achieved while painting.
 	 */
 	int32 VisualizeSelectedNodesAsRectangles( const TArray<TSharedPtr<FReflectorNode>>& InNodesToDraw, const TSharedRef<SWindow>& VisualizeInWindow, FSlateWindowElementList& OutDrawElements, int32 LayerId );
diff --git a/Engine/Source/Developer/SlateReflector/Private/Widgets/SWidgetReflectorToolTipWidget.h b/Engine/Source/Developer/SlateReflector/Private/Widgets/SWidgetReflectorToolTipWidget.h
--- a/Engine/Source/Developer/SlateReflector/Private/Widgets/SWidgetReflectorToolTipWidget.h
+++ b/Engine/Source/Developer/SlateReflector/Private/Widgets/SWidgetReflectorToolTipWidget.h
@@ -1,9 +1,5 @@
 // Copyright 1998-2014 Epic Games, Inc. All Rights Reserved.
 
-/*=============================================================================
-	SReflectorToolTipWidget.h: Declares the SReflectorToolTipWidget class.
-=============================================================================*/
-
 #pragma once
 
 
@@ -23,6 +19,8 @@ public:
 
 	SLATE_END_ARGS()
 
+public:
+
 	/** Constructs this widget with InArgs */
 	void Construct( const FArguments& InArgs )
 	{
@@ -35,10 +33,10 @@ public:
 
 			// Desired Size
 			+ SGridPanel::Slot(0, 0)
-			[
-				SNew(STextBlock)
-					.Text(LOCTEXT("DesiredSize", "Desired Size"))
-			]
+				[
+					SNew(STextBlock)
+						.Text(LOCTEXT("DesiredSize", "Desired Size"))
+				]
 
 			// Desired Size Value
 			+ SGridPanel::Slot(1, 0)
diff --git a/Engine/Source/Developer/SlateReflector/Private/Widgets/SWidgetReflectorTreeWidgetItem.h b/Engine/Source/Developer/SlateReflector/Private/Widgets/SWidgetReflectorTreeWidgetItem.h
--- a/Engine/Source/Developer/SlateReflector/Private/Widgets/SWidgetReflectorTreeWidgetItem.h
+++ b/Engine/Source/Developer/SlateReflector/Private/Widgets/SWidgetReflectorTreeWidgetItem.h
@@ -1,9 +1,5 @@
 // Copyright 1998-2014 Epic Games, Inc. All Rights Reserved.
 
-/*=============================================================================
-	SReflectorTreeWidgetItem.h: Declares the SReflectorTreeWidgetItem class.
-=============================================================================*/
-
 #pragma once
 
 
@@ -30,7 +26,7 @@ public:
 	/**
 	 * Construct child widgets that comprise this widget.
 	 *
-	 * @param InArgs  Declaration from which to construct this widget
+	 * @param InArgs Declaration from which to construct this widget.
 	 */
 	void Construct( const FArguments& InArgs, const TSharedRef<STableViewBase>& InOwnerTableView )
 	{
diff --git a/Engine/Source/Developer/SlateReflector/Public/Interfaces/ISlateReflectorModule.h b/Engine/Source/Developer/SlateReflector/Public/Interfaces/ISlateReflectorModule.h
--- a/Engine/Source/Developer/SlateReflector/Public/Interfaces/ISlateReflectorModule.h
+++ b/Engine/Source/Developer/SlateReflector/Public/Interfaces/ISlateReflectorModule.h
@@ -1,9 +1,5 @@
 // Copyright 1998-2014 Epic Games, Inc. All Rights Reserved.
 
-/*=============================================================================
-	IMessagingModule.h: Declares the IMessagingModule interface.
-=============================================================================*/
-
 #pragma once
 
 
diff --git a/Engine/Source/Developer/SlateReflector/Public/SlateReflector.h b/Engine/Source/Developer/SlateReflector/Public/SlateReflector.h
--- a/Engine/Source/Developer/SlateReflector/Public/SlateReflector.h
+++ b/Engine/Source/Developer/SlateReflector/Public/SlateReflector.h
@@ -1,9 +1,5 @@
 // Copyright 1998-2014 Epic Games, Inc. All Rights Reserved.
 
-/*=============================================================================
-	SlateReflector.h: SlateReflector module public header file.
-=============================================================================*/
-
 #pragma once
 
 
diff --git a/Engine/Source/Developer/TargetDeviceServices/Classes/TargetDeviceServiceMessages.h b/Engine/Source/Developer/TargetDeviceServices/Classes/TargetDeviceServiceMessages.h
--- a/Engine/Source/Developer/TargetDeviceServices/Classes/TargetDeviceServiceMessages.h
+++ b/Engine/Source/Developer/TargetDeviceServices/Classes/TargetDeviceServiceMessages.h
@@ -1,9 +1,5 @@
 // Copyright 1998-2014 Epic Games, Inc. All Rights Reserved.
 
-/*=============================================================================
-	TargetDeviceServiceMessages.h: Declares message types used in target device services.
-=============================================================================*/
-
 #pragma once
 
 #include "TargetDeviceServiceMessages.generated.h"
diff --git a/Engine/Source/Developer/TargetDeviceServices/Private/TargetDeviceProxy.cpp b/Engine/Source/Developer/TargetDeviceServices/Private/TargetDeviceProxy.cpp
--- a/Engine/Source/Developer/TargetDeviceServices/Private/TargetDeviceProxy.cpp
+++ b/Engine/Source/Developer/TargetDeviceServices/Private/TargetDeviceProxy.cpp
@@ -1,9 +1,5 @@
 // Copyright 1998-2014 Epic Games, Inc. All Rights Reserved.
 
-/*=============================================================================
-	TargetDeviceProxy.cpp: Implements the FTargetDeviceProxy class.
-=============================================================================*/
-
 #include "TargetDeviceServicesPrivatePCH.h"
 
 
diff --git a/Engine/Source/Developer/TargetDeviceServices/Private/TargetDeviceProxy.h b/Engine/Source/Developer/TargetDeviceServices/Private/TargetDeviceProxy.h
--- a/Engine/Source/Developer/TargetDeviceServices/Private/TargetDeviceProxy.h
+++ b/Engine/Source/Developer/TargetDeviceServices/Private/TargetDeviceProxy.h
@@ -1,20 +1,12 @@
 // Copyright 1998-2014 Epic Games, Inc. All Rights Reserved.
 
-/*=============================================================================
-	FTargetDeviceProxy.h: Declares the FTargetDeviceProxy class.
-=============================================================================*/
-
 #pragma once
 
 
-/**
- * Type definition for shared pointers to instances of FTargetDeviceProxy.
- */
+/** Type definition for shared pointers to instances of FTargetDeviceProxy. */
 typedef TSharedPtr<class FTargetDeviceProxy> FTargetDeviceProxyPtr;
 
-/**
- * Type definition for shared references to instances of FTargetDeviceProxy.
- */
+/** Type definition for shared references to instances of FTargetDeviceProxy. */
 typedef TSharedRef<class FTargetDeviceProxy> FTargetDeviceProxyRef;
 
 
@@ -34,16 +26,16 @@ public:
 	/**
 	 * Creates and initializes a new instance.
 	 *
-	 * @param InId - The identifier of the target device to create this proxy for.
+	 * @param InId The identifier of the target device to create this proxy for.
 	 */
 	FTargetDeviceProxy( const FString& InId );
 
 	/**
 	 * Creates and initializes a new instance.
 	 *
-	 * @param InId - The identifier of the target device to create this proxy for.
-	 * @param Message - The message to initialize from.
-	 * @param Context - The message context.
+	 * @param InId The identifier of the target device to create this proxy for.
+	 * @param Message The message to initialize from.
+	 * @param Context The message context.
 	 */
 	FTargetDeviceProxy( const FString& InId, const FTargetDeviceServicePong& Message, const IMessageContextRef& Context );
 
@@ -62,123 +54,121 @@ public:
 	/*
 	 * Updates the proxy's information from the given device service response.
 	 *
-	 * @param Message - The message containing the response.
-	 * @param Context - The message context.
+	 * @param Message The message containing the response.
+	 * @param Context The message context.
 	 */
 	void UpdateFromMessage( const FTargetDeviceServicePong& Message, const IMessageContextRef& Context );
 
 public:
 
-	// Begin ITargetDeviceProxy interface
+	// ITargetDeviceProxy interface
 
-	virtual const bool CanMultiLaunch( ) const OVERRIDE
+	virtual const bool CanMultiLaunch( ) const override
 	{
 		return SupportsMultiLaunch;
 	}
 
-	virtual bool CanPowerOff( ) const OVERRIDE
+	virtual bool CanPowerOff( ) const override
 	{
 		return SupportsPowerOff;
 	}
 
-	virtual bool CanPowerOn( ) const OVERRIDE
+	virtual bool CanPowerOn( ) const override
 	{
 		return SupportsPowerOn;
 	}
 
-	virtual bool CanReboot( ) const OVERRIDE
+	virtual bool CanReboot( ) const override
 	{
 		return SupportsReboot;
 	}
 
-	virtual bool DeployApp( const TMap<FString, FString>& Files, const FGuid& TransactionId ) OVERRIDE;
+	virtual bool DeployApp( const TMap<FString, FString>& Files, const FGuid& TransactionId ) override;
 
-	virtual const FString& GetDeviceId( ) OVERRIDE
+	virtual const FString& GetDeviceId( ) override
 	{
 		return Id;
 	}
 
-	virtual const FString& GetHostName( ) const OVERRIDE
+	virtual const FString& GetHostName( ) const override
 	{
 		return HostName;
 	}
 
-	virtual const FString& GetHostUser( ) const OVERRIDE
+	virtual const FString& GetHostUser( ) const override
 	{
 		return HostUser;
 	}
 
-	virtual const FString& GetMake( ) const OVERRIDE
+	virtual const FString& GetMake( ) const override
 	{
 		return Make;
 	}
 
-	virtual const FString& GetModel( ) const OVERRIDE
+	virtual const FString& GetModel( ) const override
 	{
 		return Model;
 	}
 
-	virtual const FString& GetName( ) const OVERRIDE
+	virtual const FString& GetName( ) const override
 	{
 		return Name;
 	}
 
-	virtual const FString& GetDeviceUser( ) const OVERRIDE
+	virtual const FString& GetDeviceUser( ) const override
 	{
 		return DeviceUser;
 	}
 
-	virtual const FString& GetDeviceUserPassword( ) const OVERRIDE
+	virtual const FString& GetDeviceUserPassword( ) const override
 	{
 		return DeviceUserPassword;
 	}
 
-	virtual const FString& GetPlatformName( ) const OVERRIDE
+	virtual const FString& GetPlatformName( ) const override
 	{
 		return Platform;
 	}
 
-	virtual const FString& GetType( ) const OVERRIDE
+	virtual const FString& GetType( ) const override
 	{
 		return Type;
 	}
 
-	virtual bool IsConnected( ) const OVERRIDE
+	virtual bool IsConnected( ) const override
 	{
 		return Connected;
 	}
 
-	virtual bool IsShared( ) const OVERRIDE
+	virtual bool IsShared( ) const override
 	{
 		return Shared;
 	}
 
-	virtual bool LaunchApp( const FString& AppId, EBuildConfigurations::Type BuildConfiguration, const FString& Params ) OVERRIDE;
+	virtual bool LaunchApp( const FString& AppId, EBuildConfigurations::Type BuildConfiguration, const FString& Params ) override;
 
-	virtual FOnTargetDeviceProxyDeployCommitted& OnDeployCommitted( ) OVERRIDE
+	virtual FOnTargetDeviceProxyDeployCommitted& OnDeployCommitted( ) override
 	{
 		return DeployCommittedDelegate;
 	}
 
-	virtual FOnTargetDeviceProxyDeployFailed& OnDeployFailed( ) OVERRIDE
+	virtual FOnTargetDeviceProxyDeployFailed& OnDeployFailed( ) override
 	{
 		return DeployFailedDelegate;
 	}
 
-	virtual FOnTargetDeviceProxyLaunchFailed& OnLaunchFailed( ) OVERRIDE
+	virtual FOnTargetDeviceProxyLaunchFailed& OnLaunchFailed( ) override
 	{
 		return LaunchFailedDelegate;
 	}
 
-	virtual void PowerOff( bool Force ) OVERRIDE;
-
-	virtual void PowerOn( ) OVERRIDE;
+	virtual void PowerOff( bool Force ) override;
 
-	virtual void Reboot( ) OVERRIDE;
+	virtual void PowerOn( ) override;
 
-	virtual void Run( const FString& ExecutablePath, const FString& Params ) OVERRIDE;
+	virtual void Reboot( ) override;
 
-	// End ITargetDeviceProxy interface
+	virtual void Run( const FString& ExecutablePath, const FString& Params ) override;
 
 protected:
 
diff --git a/Engine/Source/Developer/TargetDeviceServices/Private/TargetDeviceProxyManager.cpp b/Engine/Source/Developer/TargetDeviceServices/Private/TargetDeviceProxyManager.cpp
--- a/Engine/Source/Developer/TargetDeviceServices/Private/TargetDeviceProxyManager.cpp
+++ b/Engine/Source/Developer/TargetDeviceServices/Private/TargetDeviceProxyManager.cpp
@@ -1,9 +1,5 @@
 // Copyright 1998-2014 Epic Games, Inc. All Rights Reserved.
 
-/*=============================================================================
-	TargetDeviceProxyManager.cpp: Implements the FTargetDeviceProxyManager class.
-=============================================================================*/
-
 #include "TargetDeviceServicesPrivatePCH.h"
 
 
diff --git a/Engine/Source/Developer/TargetDeviceServices/Private/TargetDeviceProxyManager.h b/Engine/Source/Developer/TargetDeviceServices/Private/TargetDeviceProxyManager.h
--- a/Engine/Source/Developer/TargetDeviceServices/Private/TargetDeviceProxyManager.h
+++ b/Engine/Source/Developer/TargetDeviceServices/Private/TargetDeviceProxyManager.h
@@ -1,9 +1,5 @@
 // Copyright 1998-2014 Epic Games, Inc. All Rights Reserved.
 
-/*=============================================================================
-	TargetDeviceProxyManager.h: Declares the FTargetDeviceProxyManager class.
-=============================================================================*/
-
 #pragma once
 
 
@@ -27,28 +23,24 @@ public:
 
 public:
 
-	// Begin ITargetDeviceProxyLocator interface
-
-	virtual ITargetDeviceProxyRef FindOrAddProxy( const FString& DeviceId ) OVERRIDE;
+	// ITargetDeviceProxyLocator interface
 
-	virtual ITargetDeviceProxyPtr FindProxy( const FString& DeviceId ) OVERRIDE;
-
-	virtual void GetProxies( const FString& PlatformName, bool IncludeUnshared, TArray<ITargetDeviceProxyPtr>& OutProxies ) OVERRIDE;
+	virtual ITargetDeviceProxyRef FindOrAddProxy( const FString& DeviceId ) override;
+	virtual ITargetDeviceProxyPtr FindProxy( const FString& DeviceId ) override;
+	virtual void GetProxies( const FString& PlatformName, bool IncludeUnshared, TArray<ITargetDeviceProxyPtr>& OutProxies ) override;
 
 	DECLARE_DERIVED_EVENT(FTargetDeviceProxyManager, ITargetDeviceProxyManager::FOnTargetDeviceProxyAdded, FOnTargetDeviceProxyAdded);
-	virtual FOnTargetDeviceProxyAdded& OnProxyAdded( ) OVERRIDE
+	virtual FOnTargetDeviceProxyAdded& OnProxyAdded( ) override
 	{
 		return ProxyAddedDelegate;
 	}
 
 	DECLARE_DERIVED_EVENT(FTargetDeviceProxyManager, ITargetDeviceProxyManager::FOnTargetDeviceProxyRemoved, FOnTargetDeviceProxyRemoved);
-	virtual FOnTargetDeviceProxyRemoved& OnProxyRemoved( ) OVERRIDE
+	virtual FOnTargetDeviceProxyRemoved& OnProxyRemoved( ) override
 	{
 		return ProxyRemovedDelegate;
 	}
 
-	// End ITargetDeviceProxyLocator interface
-
 protected:
 
 	/**
diff --git a/Engine/Source/Developer/TargetDeviceServices/Private/TargetDeviceService.h b/Engine/Source/Developer/TargetDeviceServices/Private/TargetDeviceService.h
--- a/Engine/Source/Developer/TargetDeviceServices/Private/TargetDeviceService.h
+++ b/Engine/Source/Developer/TargetDeviceServices/Private/TargetDeviceService.h
@@ -1,9 +1,5 @@
 // Copyright 1998-2014 Epic Games, Inc. All Rights Reserved.
 
-/*=============================================================================
-	ClaimedTargetDevice.h: Declares and implements the FClaimedTargetDevice class.
-=============================================================================*/
-
 #pragma once
 
 
@@ -18,9 +14,9 @@ public:
 	/**
 	 * Creates and initializes a new instance.
 	 *
-	 * @param InDeviceId - The identifier of the target device to expose.
-	 * @param InDeviceName - The name of the device to expose.
-	 * @param InMessageBus - The message bus to listen on for clients.
+	 * @param InDeviceId The identifier of the target device to expose.
+	 * @param InDeviceName The name of the device to expose.
+	 * @param InMessageBus The message bus to listen on for clients.
 	 */
 	FTargetDeviceService( const FTargetDeviceId& InDeviceId, const FString& InDeviceName, const IMessageBusRef& InMessageBus )
 		: CachedDeviceName(InDeviceName)
@@ -61,9 +57,9 @@ public:
 
 public:
 
-	// Begin ITargetDeviceService interface
+	// ITargetDeviceService interface
 
-	virtual bool CanStart( ) const OVERRIDE
+	virtual bool CanStart( ) const override
 	{
 		ITargetDevicePtr TargetDevice = TargetDevicePtr.Pin();
 
@@ -75,27 +71,27 @@ public:
 		return false;
 	}
 
-	virtual const FString& GetClaimHost( ) OVERRIDE
+	virtual const FString& GetClaimHost( ) override
 	{
 		return ClaimHost;
 	}
 
-	virtual const FString& GetClaimUser( ) OVERRIDE
+	virtual const FString& GetClaimUser( ) override
 	{
 		return ClaimUser;
 	}
 
-	virtual ITargetDevicePtr GetDevice( ) OVERRIDE
+	virtual ITargetDevicePtr GetDevice( ) override
 	{
 		return AcquireDevice();
 	}
 
-	virtual const FTargetDeviceId& GetDeviceId( ) const OVERRIDE
+	virtual const FTargetDeviceId& GetDeviceId( ) const override
 	{
 		return DeviceId;
 	}
 
-	virtual FString GetCachedDeviceName( ) const OVERRIDE
+	virtual FString GetCachedDeviceName( ) const override
 	{
 		ITargetDevicePtr TargetDevice = TargetDevicePtr.Pin();
 
@@ -112,22 +108,22 @@ public:
 		return CachedDeviceName;
 	}
 
-	virtual bool IsRunning( ) const OVERRIDE
+	virtual bool IsRunning( ) const override
 	{
 		return (Running && TargetDevicePtr.IsValid());
 	}
 
-	virtual bool IsShared( ) const OVERRIDE
+	virtual bool IsShared( ) const override
 	{
 		return (Running && Shared);
 	}
 
-	virtual void SetShared( bool InShared ) OVERRIDE
+	virtual void SetShared( bool InShared ) override
 	{
 		Shared = InShared;
 	}
 
-	virtual bool Start( ) OVERRIDE
+	virtual bool Start( ) override
 	{
 		if (!Running && MessageEndpoint.IsValid())
 		{
@@ -154,7 +150,7 @@ public:
 		return true;
 	}
 
-	virtual void Stop( ) OVERRIDE
+	virtual void Stop( ) override
 	{
 		if (Running)
 		{
@@ -173,8 +169,6 @@ public:
 		}
 	}
 
-	// End ITargetDeviceService interface
-
 protected:
 
 	/**
@@ -198,6 +192,8 @@ protected:
 	/**
 	 * Stores the specified file to deploy.
 	 *
+	 * @param FileReader The archive reader providing access to the file data.
+	 * @param TargetFileName The desired name of the file on the target device.
 	 */
 	bool StoreDeployedFile( FArchive* FileReader, const FString& TargetFileName  ) const
 	{
diff --git a/Engine/Source/Developer/TargetDeviceServices/Private/TargetDeviceServiceManager.h b/Engine/Source/Developer/TargetDeviceServices/Private/TargetDeviceServiceManager.h
--- a/Engine/Source/Developer/TargetDeviceServices/Private/TargetDeviceServiceManager.h
+++ b/Engine/Source/Developer/TargetDeviceServices/Private/TargetDeviceServiceManager.h
@@ -1,9 +1,5 @@
 // Copyright 1998-2014 Epic Games, Inc. All Rights Reserved.
 
-/*=============================================================================
-	TargetDeviceServiceManager.h: Declares the FTargetDeviceServiceManager class.
-=============================================================================*/
-
 #pragma once
 
 
@@ -29,36 +25,32 @@ public:
 
 public:
 
-	// Begin ITargetDeviceServiceManager interface
-
-	virtual bool AddStartupService( const FTargetDeviceId& DeviceId, const FString& PreliminaryDeviceName ) OVERRIDE;
+	// ITargetDeviceServiceManager interface
 
-	virtual int32 GetServices( TArray<ITargetDeviceServicePtr>& OutServices ) OVERRIDE;
+	virtual bool AddStartupService( const FTargetDeviceId& DeviceId, const FString& PreliminaryDeviceName ) override;
+	virtual int32 GetServices( TArray<ITargetDeviceServicePtr>& OutServices ) override;
 
 	DECLARE_DERIVED_EVENT(FTargetDeviceServiceManager, ITargetDeviceServiceManager::FOnTargetDeviceServiceAdded, FOnTargetDeviceServiceAdded);
-	virtual FOnTargetDeviceServiceAdded& OnServiceAdded( ) OVERRIDE
+	virtual FOnTargetDeviceServiceAdded& OnServiceAdded( ) override
 	{
 		return ServiceAddedDelegate;
 	}
 
 	DECLARE_DERIVED_EVENT(FTargetDeviceServiceManager, ITargetDeviceServiceManager::FOnTargetDeviceServiceRemoved, FOnTargetDeviceServiceRemoved);
-	virtual FOnTargetDeviceServiceRemoved& OnServiceRemoved( ) OVERRIDE
+	virtual FOnTargetDeviceServiceRemoved& OnServiceRemoved( ) override
 	{
 		return ServiceRemovedDelegate;
 	}
 
-	virtual void RemoveStartupService( const FTargetDeviceId& DeviceId ) OVERRIDE;
-
-	// End ITargetDeviceServiceManager interface
+	virtual void RemoveStartupService( const FTargetDeviceId& DeviceId ) override;
 
 protected:
 
 	/**
 	 * Adds a device service for the given target device.
 	 *
-	 * @param DeviceId - The identifier of the device to add a service for.
-	 * @param PreliminaryDeviceName - The preliminary name to assign to this device.
-	 *
+	 * @param DeviceId The identifier of the device to add a service for.
+	 * @param PreliminaryDeviceName The preliminary name to assign to this device.
 	 * @return true if the service was created, false otherwise.
 	 */
 	bool AddService( const FTargetDeviceId& DeviceId, const FString& PreliminaryDeviceName );
@@ -80,7 +72,7 @@ protected:
 	/**
 	 * Removes the device service for the given target device.
 	 *
-	 * @param DeviceId - The identifier of the device to remove the service for.
+	 * @param DeviceId The identifier of the device to remove the service for.
 	 */
 	void RemoveService( const FTargetDeviceId& DeviceId );
 
diff --git a/Engine/Source/Developer/TargetDeviceServices/Private/TargetDeviceServicesModule.cpp b/Engine/Source/Developer/TargetDeviceServices/Private/TargetDeviceServicesModule.cpp
--- a/Engine/Source/Developer/TargetDeviceServices/Private/TargetDeviceServicesModule.cpp
+++ b/Engine/Source/Developer/TargetDeviceServices/Private/TargetDeviceServicesModule.cpp
@@ -1,13 +1,8 @@
 // Copyright 1998-2014 Epic Games, Inc. All Rights Reserved.
 
-/*=============================================================================
-	TargetDeviceServicesModule.cpp: Implements the FTargetDeviceServicesModule class.
-=============================================================================*/
-
 #include "TargetDeviceServicesPrivatePCH.h"
 
 
-
 DEFINE_LOG_CATEGORY(TargetDeviceServicesLog);
 
 
@@ -19,9 +14,9 @@ class FTargetDeviceServicesModule
 {
 public:
 
-	// Begin ITargetDeviceServicesModule interface
+	// ITargetDeviceServicesModule interface
 
-	virtual ITargetDeviceProxyManagerRef GetDeviceProxyManager( ) OVERRIDE
+	virtual ITargetDeviceProxyManagerRef GetDeviceProxyManager( ) override
 	{
 		if (!DeviceProxyManagerSingleton.IsValid())
 		{
@@ -31,7 +26,7 @@ public:
 		return DeviceProxyManagerSingleton.ToSharedRef();
 	}
 
-	virtual ITargetDeviceServiceManagerRef GetDeviceServiceManager( ) OVERRIDE
+	virtual ITargetDeviceServiceManagerRef GetDeviceServiceManager( ) override
 	{
 		if (!DeviceServiceManagerSingleton.IsValid())
 		{
@@ -41,8 +36,6 @@ public:
 		return DeviceServiceManagerSingleton.ToSharedRef();
 	}
 
-	// End ITargetDeviceServicesModule interface
-
 private:
 
 	// Holds the device proxy manager singleton.
diff --git a/Engine/Source/Developer/TargetDeviceServices/Private/TargetDeviceServicesPrivatePCH.h b/Engine/Source/Developer/TargetDeviceServices/Private/TargetDeviceServicesPrivatePCH.h
--- a/Engine/Source/Developer/TargetDeviceServices/Private/TargetDeviceServicesPrivatePCH.h
+++ b/Engine/Source/Developer/TargetDeviceServices/Private/TargetDeviceServicesPrivatePCH.h
@@ -1,13 +1,9 @@
 // Copyright 1998-2014 Epic Games, Inc. All Rights Reserved.
 
-/*=============================================================================
-	TargetDeviceServicesPrivatePCH.h: Pre-compiled header file for the TargetDeviceServices module.
-=============================================================================*/
-
 #pragma once
 
 
-#include "../Public/TargetDeviceServices.h"
+#include "TargetDeviceServices.h"
 
 
 /* Dependencies
@@ -23,14 +19,10 @@
 /* Private declarations
  *****************************************************************************/
 
-/**
- * Declares the module's log category.
- */
+/** Declares the module's log category. */
 DECLARE_LOG_CATEGORY_EXTERN(TargetDeviceServicesLog, Log, All);
 
-/**
- * Defines the interval in seconds in which devices are being pinged by the proxy manager.
- */
+/** Defines the interval in seconds in which devices are being pinged by the proxy manager. */
 #define TARGET_DEVICE_SERVICES_PING_INTERVAL 2.5f
 
 
diff --git a/Engine/Source/Developer/TargetDeviceServices/Public/Interfaces/ITargetDeviceProxy.h b/Engine/Source/Developer/TargetDeviceServices/Public/Interfaces/ITargetDeviceProxy.h
--- a/Engine/Source/Developer/TargetDeviceServices/Public/Interfaces/ITargetDeviceProxy.h
+++ b/Engine/Source/Developer/TargetDeviceServices/Public/Interfaces/ITargetDeviceProxy.h
@@ -1,20 +1,12 @@
 // Copyright 1998-2014 Epic Games, Inc. All Rights Reserved.
 
-/*=============================================================================
-	ITargetDeviceProxy.h: Declares the ITargetDeviceProxy interface.
-=============================================================================*/
-
 #pragma once
 
 
-/**
- * Type definition for shared pointers to instances of ITargetDeviceProxy.
- */
+/** Type definition for shared pointers to instances of ITargetDeviceProxy. */
 typedef TSharedPtr<class ITargetDeviceProxy> ITargetDeviceProxyPtr;
 
-/**
- * Type definition for shared references to instances of ITargetDeviceProxy.
- */
+/** Type definition for shared references to instances of ITargetDeviceProxy. */
 typedef TSharedRef<class ITargetDeviceProxy> ITargetDeviceProxyRef;
 
 
@@ -108,7 +100,7 @@ public:
 	/**
 	 * Gets the name of the (device's) user that is logged in on a device.
 	 *
-	 * @return Devie user name.
+	 * @return Device user name.
 	 */
 	virtual const FString& GetDeviceUser( ) const = 0;
 
@@ -177,11 +169,9 @@ public:
 	 * application identifier (AppId) that can be used to launch to the deployed
 	 * application. If deploymend fails, the OnDeployFailed() delegate will be invoked.
 	 *
-	 * @param Files - The files to deploy.
-	 * @param TransactionId - The transaction identifier for the deployment.
-	 *
+	 * @param Files The files to deploy.
+	 * @param TransactionId The transaction identifier for the deployment.
 	 * @return true if deployment has been started, false otherwise.
-	 *
 	 * @see Launch
 	 * @see OnDeployCommitted
 	 * @see OnDeployFailed
@@ -193,12 +183,10 @@ public:
 	 *
 	 * If the launching fails, the OnLaunchFailed() delegate will be invoked.
 	 *
-	 * @param AppId - The identifier of the application to launch.
-	 * @param BuildConfiguration - The build configuration to launch with.
-	 * @param Params - The command line parameters to launch with.
-	 *
+	 * @param AppId The identifier of the application to launch.
+	 * @param BuildConfiguration The build configuration to launch with.
+	 * @param Params The command line parameters to launch with.
 	 * @return true on success, false otherwise.
-	 *
 	 * @see Deploy
 	 * @see OnLaunchFailed
 	 */
@@ -207,8 +195,7 @@ public:
 	/**
 	 * Powers off the device.
 	 *
-	 * @param Force - Whether to force powering off.
-	 *
+	 * @param Force Whether to force powering off.
 	 * @see PowerOn
 	 * @see Reboot
 	 */
@@ -233,8 +220,8 @@ public:
 	/**
 	 * Attempts to run an executable on the device.
 	 *
-	 * @param ExecutablePath - The path to the executable to run.
-	 * @param Params - The command line parameters.
+	 * @param ExecutablePath The path to the executable to run.
+	 * @param Params The command line parameters.
 	 */
 	virtual void Run( const FString& ExecutablePath, const FString& Params ) = 0;
 
@@ -244,7 +231,6 @@ public:
 	 * Returns a delegate to be invoked when a build has been deployed to the target device.
 	 *
 	 * @return The delegate.
-	 *
 	 * @see DeployApp
 	 */
 	virtual FOnTargetDeviceProxyDeployCommitted& OnDeployCommitted( ) = 0;
@@ -253,7 +239,6 @@ public:
 	 * Returns a delegate to be invoked when a build failed to deploy to the target device.
 	 *
 	 * @return The delegate.
-	 *
 	 * @see DeployApp
 	 */
 	virtual FOnTargetDeviceProxyDeployFailed& OnDeployFailed( ) = 0;
@@ -262,7 +247,6 @@ public:
 	 * Returns a delegate to be invoked when a build failed to launch on the target device.
 	 *
 	 * @return The delegate.
-	 *
 	 * @see LaunchApp
 	 */
 	virtual FOnTargetDeviceProxyLaunchFailed& OnLaunchFailed( ) = 0;
diff --git a/Engine/Source/Developer/TargetDeviceServices/Public/Interfaces/ITargetDeviceProxyManager.h b/Engine/Source/Developer/TargetDeviceServices/Public/Interfaces/ITargetDeviceProxyManager.h
--- a/Engine/Source/Developer/TargetDeviceServices/Public/Interfaces/ITargetDeviceProxyManager.h
+++ b/Engine/Source/Developer/TargetDeviceServices/Public/Interfaces/ITargetDeviceProxyManager.h
@@ -1,20 +1,12 @@
 // Copyright 1998-2014 Epic Games, Inc. All Rights Reserved.
 
-/*=============================================================================
-	ITargetDeviceProxyManager.h: Declares the ITargetDeviceProxyManager interface.
-=============================================================================*/
-
 #pragma once
 
 
-/**
- * Type definition for shared pointers to instances of ITargetDeviceProxyManager.
- */
+/** Type definition for shared pointers to instances of ITargetDeviceProxyManager. */
 typedef TSharedPtr<class ITargetDeviceProxyManager> ITargetDeviceProxyManagerPtr;
 
-/**
- * Type definition for shared references to instances of ITargetDeviceProxyManager.
- */
+/** Type definition for shared references to instances of ITargetDeviceProxyManager. */
 typedef TSharedRef<class ITargetDeviceProxyManager> ITargetDeviceProxyManagerRef;
 
 
@@ -28,8 +20,7 @@ public:
 	/**
 	 * Finds or adds a proxy for the specified device identifier.
 	 *
-	 * @param DeviceId - The identifier of the device to create the proxy for.
-	 *
+	 * @param DeviceId The identifier of the device to create the proxy for.
 	 * @return The proxy.
 	 */
 	virtual ITargetDeviceProxyRef FindOrAddProxy( const FString& DeviceId ) = 0;
@@ -37,8 +28,7 @@ public:
 	/**
 	 * Finds the proxy for the specified device identifier.
 	 *
-	 * @param DeviceId - The identifier of the device to create the proxy for.
-	 *
+	 * @param DeviceId The identifier of the device to create the proxy for.
 	 * @return The proxy, or nullptr if it couldn't be found.
 	 */
 	virtual ITargetDeviceProxyPtr FindProxy( const FString& DeviceId ) = 0;
@@ -46,9 +36,9 @@ public:
 	/**
 	 * Gets a list of devices found by the device discovery.
 	 *
-	 * @param PlatformName - The the name of the target platform to get proxies for (or empty string for all proxies).
-	 * @param IncludeUnshared - Whether to include devices that are not being shared with the local user.
-	 * @param OutProxies - Will hold the list of devices found by the locator
+	 * @param PlatformName The the name of the target platform to get proxies for (or empty string for all proxies).
+	 * @param IncludeUnshared Whether to include devices that are not being shared with the local user.
+	 * @param OutProxies Will hold the list of devices found by the locator
 	 */
 	virtual void GetProxies( const FString& PlatformName, bool IncludeUnshared, TArray<ITargetDeviceProxyPtr>& OutProxies ) = 0;
 
diff --git a/Engine/Source/Developer/TargetDeviceServices/Public/Interfaces/ITargetDeviceService.h b/Engine/Source/Developer/TargetDeviceServices/Public/Interfaces/ITargetDeviceService.h
--- a/Engine/Source/Developer/TargetDeviceServices/Public/Interfaces/ITargetDeviceService.h
+++ b/Engine/Source/Developer/TargetDeviceServices/Public/Interfaces/ITargetDeviceService.h
@@ -1,20 +1,12 @@
 // Copyright 1998-2014 Epic Games, Inc. All Rights Reserved.
 
-/*=============================================================================
-	ITargetDeviceService.h: Declares the ITargetDeviceService interface.
-=============================================================================*/
-
 #pragma once
 
 
-/**
- * Type definition for shared pointers to instances of ITargetDeviceService.
- */
+/** Type definition for shared pointers to instances of ITargetDeviceService. */
 typedef TSharedPtr<class ITargetDeviceService, ESPMode::ThreadSafe> ITargetDeviceServicePtr;
 
-/**
- * Type definition for shared references to instances of ITargetDeviceService.
- */
+/** Type definition for shared references to instances of ITargetDeviceService. */
 typedef TSharedRef<class ITargetDeviceService, ESPMode::ThreadSafe> ITargetDeviceServiceRef;
 
 
@@ -90,7 +82,6 @@ public:
 	 * can discover and use the device.
 	 *
 	 * @return true if the device is shared, false otherwise.
-	 *
 	 * @see SetShared
 	 */
 	virtual bool IsShared( ) const = 0;
@@ -103,8 +94,7 @@ public:
 	 * If a device is shared, the device proxies of all users on the network
 	 * can discover and use the device.
 	 *
-	 * @param InShared - Indicates whether device sharing is enabled.
-	 *
+	 * @param InShared Indicates whether device sharing is enabled.
 	 * @see IsShared
 	 */
 	virtual void SetShared( bool InShared ) = 0;
@@ -113,7 +103,6 @@ public:
 	 * Starts the service.
 	 *
 	 * @return true if the service has been started, false otherwise.
-	 *
 	 * @see IsRunning
 	 * @see Stop
 	 */
diff --git a/Engine/Source/Developer/TargetDeviceServices/Public/Interfaces/ITargetDeviceServiceManager.h b/Engine/Source/Developer/TargetDeviceServices/Public/Interfaces/ITargetDeviceServiceManager.h
--- a/Engine/Source/Developer/TargetDeviceServices/Public/Interfaces/ITargetDeviceServiceManager.h
+++ b/Engine/Source/Developer/TargetDeviceServices/Public/Interfaces/ITargetDeviceServiceManager.h
@@ -1,20 +1,12 @@
 // Copyright 1998-2014 Epic Games, Inc. All Rights Reserved.
 
-/*=============================================================================
-	ITargetDeviceServiceManager.h: Declares the ITargetDeviceServiceManager interface.
-=============================================================================*/
-
 #pragma once
 
 
-/**
- * Type definition for shared pointers to instances of ITargetDeviceServiceManager.
- */
+/** Type definition for shared pointers to instances of ITargetDeviceServiceManager. */
 typedef TSharedPtr<class ITargetDeviceServiceManager> ITargetDeviceServiceManagerPtr;
 
-/**
- * Type definition for shared references to instances of ITargetDeviceServiceManager.
- */
+/** Type definition for shared references to instances of ITargetDeviceServiceManager. */
 typedef TSharedRef<class ITargetDeviceServiceManager> ITargetDeviceServiceManagerRef;
 
 
@@ -32,11 +24,9 @@ public:
 	 * could not be discovered automatically or are currently unavailable. This name will
 	 * be replaced with the actual device name once the physical device becomes available.
 	 *
-	 * @param DeviceId - The device identifier for the service to add.
-	 * @param PreliminaryDeviceName - The preliminary device name.
-	 *
+	 * @param DeviceId The device identifier for the service to add.
+	 * @param PreliminaryDeviceName The preliminary device name.
 	 * @return true if the service added, false otherwise.
-	 *
 	 * @see GetServices
 	 * @see RemoveStartupService
 	 */
@@ -45,8 +35,7 @@ public:
 	/**
 	 * Gets the collection of target device services managed by this instance.
 	 *
-	 * @param OutServices - Will contain the collection of services.
-	 *
+	 * @param OutServices Will contain the collection of services.
 	 * @return The number of services returned.
 	 */
 	virtual int32 GetServices( TArray<ITargetDeviceServicePtr>& OutServices ) = 0;
@@ -54,8 +43,7 @@ public:
 	/**
 	 * Removes a service from the list of services that are started automatically.
 	 *
-	 * @param DeviceId - The device identifier for the service to remove.
-	 *
+	 * @param DeviceId The device identifier for the service to remove.
 	 * @see AddStartupService
 	 * @see GetServices
 	 */
diff --git a/Engine/Source/Developer/TargetDeviceServices/Public/Interfaces/ITargetDeviceServicesModule.h b/Engine/Source/Developer/TargetDeviceServices/Public/Interfaces/ITargetDeviceServicesModule.h
--- a/Engine/Source/Developer/TargetDeviceServices/Public/Interfaces/ITargetDeviceServicesModule.h
+++ b/Engine/Source/Developer/TargetDeviceServices/Public/Interfaces/ITargetDeviceServicesModule.h
@@ -1,9 +1,5 @@
 // Copyright 1998-2014 Epic Games, Inc. All Rights Reserved.
 
-/*=============================================================================
-	ITargetDeviceServicesModule.h: Declares the ITargetDeviceServicesModule interface.
-=============================================================================*/
-
 #pragma once
 
 
diff --git a/Engine/Source/Developer/TargetDeviceServices/Public/TargetDeviceServices.h b/Engine/Source/Developer/TargetDeviceServices/Public/TargetDeviceServices.h
--- a/Engine/Source/Developer/TargetDeviceServices/Public/TargetDeviceServices.h
+++ b/Engine/Source/Developer/TargetDeviceServices/Public/TargetDeviceServices.h
@@ -1,9 +1,5 @@
 // Copyright 1998-2014 Epic Games, Inc. All Rights Reserved.
 
-/*=============================================================================
-	ITargetDeviceServices.h: ITargetDeviceServices module public header file.
-=============================================================================*/
-
 #pragma once
 
 
diff --git a/Engine/Source/Developer/TargetPlatform/Private/TargetPlatformPrivatePCH.h b/Engine/Source/Developer/TargetPlatform/Private/TargetPlatformPrivatePCH.h
--- a/Engine/Source/Developer/TargetPlatform/Private/TargetPlatformPrivatePCH.h
+++ b/Engine/Source/Developer/TargetPlatform/Private/TargetPlatformPrivatePCH.h
@@ -1,20 +1,12 @@
 // Copyright 1998-2014 Epic Games, Inc. All Rights Reserved.
 
-/*=============================================================================
-	TargetPlatformPrivatePCH.h: Pre-compiled header file for the TargetPlatform module.
-=============================================================================*/
-
 #pragma once
 
+#include "TargetPlatform.h"
 
-/* Dependencies
- *****************************************************************************/
 
-#include "../Public/TargetPlatform.h"
+/* Private dependencies
+ *****************************************************************************/
 
 #include "Core.h"
 #include "ModuleManager.h"
-
-
-/* Private includes
- *****************************************************************************/
diff --git a/Engine/Source/Editor/DetailCustomizations/Private/DeviceProfileDetails.cpp b/Engine/Source/Editor/DetailCustomizations/Private/DeviceProfileDetails.cpp
--- a/Engine/Source/Editor/DetailCustomizations/Private/DeviceProfileDetails.cpp
+++ b/Engine/Source/Editor/DetailCustomizations/Private/DeviceProfileDetails.cpp
@@ -1,9 +1,5 @@
 // Copyright 1998-2014 Epic Games, Inc. All Rights Reserved.
 
-/*=============================================================================
-	DeviceProfileDetails.cpp: Implements the FDeviceProfileDetails class.
-=============================================================================*/
-
 #include "DetailCustomizationsPrivatePCH.h"
 #include "DeviceProfiles/DeviceProfile.h"
 #include "DeviceProfiles/DeviceProfileManager.h"
diff --git a/Engine/Source/Editor/DetailCustomizations/Private/DeviceProfileDetails.h b/Engine/Source/Editor/DetailCustomizations/Private/DeviceProfileDetails.h
--- a/Engine/Source/Editor/DetailCustomizations/Private/DeviceProfileDetails.h
+++ b/Engine/Source/Editor/DetailCustomizations/Private/DeviceProfileDetails.h
@@ -1,13 +1,10 @@
 // Copyright 1998-2014 Epic Games, Inc. All Rights Reserved.
 
-/*=============================================================================
-	DeviceProfileDetails.h: Declares the FDeviceProfileDetails class.
-=============================================================================*/
-
 #pragma once
 
 class UDeviceProfile;
 
+
 /* FDeviceProfileParentPropertyDetails
  *****************************************************************************/
 
@@ -23,7 +20,7 @@ public:
 	/**
 	 * Constructor for the parent property details view
 	 *
-	 * @param InDetailsBuilder - Where we are adding our property view to
+	 * @param InDetailsBuilder Where we are adding our property view to
 	 */
 	FDeviceProfileParentPropertyDetails(IDetailLayoutBuilder* InDetailBuilder);
 
@@ -38,22 +35,19 @@ private:
 	/**
 	 * Select a parent Combo-box listings
 	 *
-	 * @param InItem - The prospective parent profile name we are creating a listing for
-	 *
+	 * @param InItem The prospective parent profile name we are creating a listing for
 	 * @return The slate widget with the listed profile name
 	 */
 	TSharedRef<SWidget> HandleDeviceProfileParentComboBoxGenarateWidget(TSharedPtr<FString> InItem);
 
-
 	/**
 	 * Handle a new parent profile being selected for the active device profile
 	 *
-	 * @param NewSelection	- The name of the selected parent profile
-	 * @param SelectionInfo	- The selection type information
+	 * @param NewSelection The name of the selected parent profile
+	 * @param SelectionInfo	The selection type information
 	 */
 	void HandleDeviceProfileParentSelectionChanged(TSharedPtr<FString> NewSelection, ESelectInfo::Type SelectInfo);
 
-
 	/**
 	* Delegate used when the device profiles parent is updated from any source.
 	*/
@@ -90,7 +84,7 @@ public:
 	/**
 	 * Constructor for the active profiles' console variables property details view
 	 *
-	 * @param InDetailsBuilder - Where we are adding our property view to
+	 * @param InDetailsBuilder Where we are adding our property view to
 	 */
 	FDeviceProfileConsoleVariablesPropertyDetails(IDetailLayoutBuilder* InDetailBuilder);
 
@@ -105,7 +99,7 @@ private:
 	/**
 	 * Action when a CVar has been selected for add to the device profile
 	 *
-	 * @param SelectedCVar - The CVar added
+	 * @param SelectedCVar The CVar added
 	 */
 	void HandleCVarAdded(const FString& SelectedCVar);
 
@@ -113,8 +107,8 @@ private:
 	/**
 	 * Action when a CVar has been selected for add to the device profile
 	 *
-	 * @param InProperty	- The property handle of which we are creating a row for.
-	 * @param InGroup	- The group the property is to be listed under.
+	 * @param InProperty The property handle of which we are creating a row for.
+	 * @param InGroup The group the property is to be listed under.
 	 */
 	void CreateRowWidgetForCVarProperty(TSharedPtr<IPropertyHandle> InProperty, IDetailGroup& InGroup) const;
 
@@ -122,9 +116,9 @@ private:
 	/**
 	 * Action when a CVar has been updated on the device profile
 	 *
-	 * @param NewValue				- The New value for the given CVar
-	 * @param CommitInfo			- Details of how the action was commited
-	 * @param CVarPropertyHandle	- The CVar property which has been updated
+	 * @param NewValue The New value for the given CVar
+	 * @param CommitInfo Details of how the action was committed
+	 * @param CVarPropertyHandle The CVar property which has been updated
 	 */
 	void OnCVarValueCommited(const FText& NewValue, ETextCommit::Type CommitInfo, TSharedPtr<IPropertyHandle> CVarPropertyHandle);
 
@@ -132,8 +126,7 @@ private:
 	/**
 	 * Action when a CVar has removed from the device profile
 	 *
-	 * @param InProperty - The property which has been removed from the profile
-	 *
+	 * @param InProperty The property which has been removed from the profile
 	 * @return Whether the event was handled
 	 */
 	FReply OnRemoveCVarProperty(TSharedPtr<IPropertyHandle> InProperty);
@@ -143,7 +136,6 @@ private:
 	 * Action when a group of CVars have been removed from the device profile
 	 *
 	 * @param GroupName - The name of the group all Cvars have been removed for.
-	 *
 	 * @return Whether the event was handled
 	 */
 	FReply OnRemoveAllFromGroup(const TSharedPtr<FString> GroupName);
@@ -190,11 +182,9 @@ public:
 
 public:
 
-	// Begin IDetailCustomization interface
-
-	virtual void CustomizeDetails( IDetailLayoutBuilder& InDetailBuilder ) OVERRIDE;
+	// IDetailCustomization interface
 
-	// End IDetailCustomization interface
+	virtual void CustomizeDetails( IDetailLayoutBuilder& InDetailBuilder ) override;
 
 private:
 
diff --git a/Engine/Source/Editor/DetailCustomizations/Private/EnvQueryParamInstanceCustomization.h b/Engine/Source/Editor/DetailCustomizations/Private/EnvQueryParamInstanceCustomization.h
--- a/Engine/Source/Editor/DetailCustomizations/Private/EnvQueryParamInstanceCustomization.h
+++ b/Engine/Source/Editor/DetailCustomizations/Private/EnvQueryParamInstanceCustomization.h
@@ -8,13 +8,15 @@ class FEnvQueryParamInstanceCustomization : public IPropertyTypeCustomization
 {
 public:
 
-	// Begin IPropertyTypeCustomization interface
-	virtual void CustomizeHeader( TSharedRef<class IPropertyHandle> StructPropertyHandle, class FDetailWidgetRow& HeaderRow, IPropertyTypeCustomizationUtils& StructCustomizationUtils ) OVERRIDE;
-	virtual void CustomizeChildren( TSharedRef<class IPropertyHandle> StructPropertyHandle, class IDetailChildrenBuilder& StructBuilder, IPropertyTypeCustomizationUtils& StructCustomizationUtils ) OVERRIDE;
-	// End IPropertyTypeCustomization interface
-
 	static TSharedRef<IPropertyTypeCustomization> MakeInstance( );
 
+public:
+
+	// IPropertyTypeCustomization interface
+
+	virtual void CustomizeHeader( TSharedRef<class IPropertyHandle> StructPropertyHandle, class FDetailWidgetRow& HeaderRow, IPropertyTypeCustomizationUtils& StructCustomizationUtils ) override;
+	virtual void CustomizeChildren( TSharedRef<class IPropertyHandle> StructPropertyHandle, class IDetailChildrenBuilder& StructBuilder, IPropertyTypeCustomizationUtils& StructCustomizationUtils ) override;
+
 protected:
 
 	TSharedPtr<IPropertyHandle> NameProp;
diff --git a/Engine/Source/Editor/DetailCustomizations/Private/GuidStructCustomization.cpp b/Engine/Source/Editor/DetailCustomizations/Private/GuidStructCustomization.cpp
--- a/Engine/Source/Editor/DetailCustomizations/Private/GuidStructCustomization.cpp
+++ b/Engine/Source/Editor/DetailCustomizations/Private/GuidStructCustomization.cpp
@@ -1,9 +1,5 @@
 // Copyright 1998-2014 Epic Games, Inc. All Rights Reserved.
 
-/*=============================================================================
-	GuidStructCustomization.cpp: Implements the FGuidStructCustomization class.
-=============================================================================*/
-
 #include "DetailCustomizationsPrivatePCH.h"
 #include "GuidStructCustomization.h"
 
diff --git a/Engine/Source/Editor/DetailCustomizations/Private/GuidStructCustomization.h b/Engine/Source/Editor/DetailCustomizations/Private/GuidStructCustomization.h
--- a/Engine/Source/Editor/DetailCustomizations/Private/GuidStructCustomization.h
+++ b/Engine/Source/Editor/DetailCustomizations/Private/GuidStructCustomization.h
@@ -1,9 +1,5 @@
 // Copyright 1998-2014 Epic Games, Inc. All Rights Reserved.
 
-/*=============================================================================
-	GuidStructCustomization.h: Declares the FGuidCustomization class.
-=============================================================================*/
-
 #pragma once
 
 
@@ -14,30 +10,27 @@ namespace EPropertyEditorGuidActions
 	 */
 	enum Type
 	{
-		// Generate a new GUID.
+		/** Generate a new GUID. */
 		Generate,
 
-		// Set a null GUID.
+		/** Set a null GUID. */
 		Invalidate
 	};
 }
 
 
 /**
- * Implements a details panel customization for  FGuid structures.
+ * Implements a details panel customization for FGuid structures.
  */
 class FGuidStructCustomization
 	: public IPropertyTypeCustomization
 {
 public:
 
-	// Begin IPropertyTypeCustomization interface
-
-	virtual void CustomizeHeader( TSharedRef<class IPropertyHandle> StructPropertyHandle, class FDetailWidgetRow& HeaderRow, IPropertyTypeCustomizationUtils& StructCustomizationUtils ) OVERRIDE;
-
-	virtual void CustomizeChildren( TSharedRef<class IPropertyHandle> StructPropertyHandle, class IDetailChildrenBuilder& StructBuilder, IPropertyTypeCustomizationUtils& StructCustomizationUtils ) OVERRIDE;
+	// IPropertyTypeCustomization interface
 
-	// End IPropertyTypeCustomization interface
+	virtual void CustomizeHeader( TSharedRef<class IPropertyHandle> StructPropertyHandle, class FDetailWidgetRow& HeaderRow, IPropertyTypeCustomizationUtils& StructCustomizationUtils ) override;
+	virtual void CustomizeChildren( TSharedRef<class IPropertyHandle> StructPropertyHandle, class IDetailChildrenBuilder& StructBuilder, IPropertyTypeCustomizationUtils& StructCustomizationUtils ) override;
 
 public:
 
@@ -53,7 +46,7 @@ protected:
 	/**
 	 * Sets the property's value.
 	 *
-	 * @param Guid - The value to set.
+	 * @param Guid The value to set.
 	 */
 	void SetGuidValue( const FGuid& Guid );
 
diff --git a/Engine/Source/Editor/DetailCustomizations/Private/KeyStructCustomization.cpp b/Engine/Source/Editor/DetailCustomizations/Private/KeyStructCustomization.cpp
--- a/Engine/Source/Editor/DetailCustomizations/Private/KeyStructCustomization.cpp
+++ b/Engine/Source/Editor/DetailCustomizations/Private/KeyStructCustomization.cpp
@@ -1,9 +1,5 @@
 // Copyright 1998-2014 Epic Games, Inc. All Rights Reserved.
 
-/*=============================================================================
-	KeyStructCustomization.cpp: Implements the FKeyStructCustomization class.
-=============================================================================*/
-
 #include "DetailCustomizationsPrivatePCH.h"
 #include "KeyStructCustomization.h"
 #include "ScopedTransaction.h"
diff --git a/Engine/Source/Editor/DetailCustomizations/Private/KeyStructCustomization.h b/Engine/Source/Editor/DetailCustomizations/Private/KeyStructCustomization.h
--- a/Engine/Source/Editor/DetailCustomizations/Private/KeyStructCustomization.h
+++ b/Engine/Source/Editor/DetailCustomizations/Private/KeyStructCustomization.h
@@ -1,9 +1,5 @@
 // Copyright 1998-2014 Epic Games, Inc. All Rights Reserved.
 
-/*=============================================================================
-	KeyStructCustomization.h: Declares the FGuidCustomization class.
-=============================================================================*/
-
 #pragma once
 
 /**
@@ -14,12 +10,12 @@ class FKeyStructCustomization
 {
 public:
 
-	// Begin IPropertyTypeCustomization interface
+	// IPropertyTypeCustomization interface
 
-	virtual void CustomizeHeader(TSharedRef<class IPropertyHandle> StructPropertyHandle, class FDetailWidgetRow& HeaderRow, IPropertyTypeCustomizationUtils& StructCustomizationUtils) OVERRIDE;
-	virtual void CustomizeChildren(TSharedRef<IPropertyHandle> StructPropertyHandle, class IDetailChildrenBuilder& ChildBuilder, IPropertyTypeCustomizationUtils& StructCustomizationUtils) OVERRIDE {};
+	virtual void CustomizeHeader(TSharedRef<class IPropertyHandle> StructPropertyHandle, class FDetailWidgetRow& HeaderRow, IPropertyTypeCustomizationUtils& StructCustomizationUtils) override;
+	virtual void CustomizeChildren(TSharedRef<IPropertyHandle> StructPropertyHandle, class IDetailChildrenBuilder& ChildBuilder, IPropertyTypeCustomizationUtils& StructCustomizationUtils) override { };
 
-	// End IPropertyTypeCustomization interface
+public:
 
 	/**
 	 * Creates a new instance.
@@ -37,5 +33,5 @@ private:
 	TSharedPtr<IPropertyHandle> PropertyHandle;
 	TSharedPtr<STextBlock> TextBlock;
 
-	TArray< TSharedPtr<FKey> > InputKeys;
+	TArray<TSharedPtr<FKey>> InputKeys;
 };
diff --git a/Engine/Source/Editor/DetailCustomizations/Private/NavLinkStructCustomization.h b/Engine/Source/Editor/DetailCustomizations/Private/NavLinkStructCustomization.h
--- a/Engine/Source/Editor/DetailCustomizations/Private/NavLinkStructCustomization.h
+++ b/Engine/Source/Editor/DetailCustomizations/Private/NavLinkStructCustomization.h
@@ -5,10 +5,18 @@
 class FNavLinkStructCustomization : public IPropertyTypeCustomization
 {
 public:
-	// Begin IPropertyTypeCustomization interface
-	virtual void CustomizeHeader( TSharedRef<IPropertyHandle> StructPropertyHandle, class FDetailWidgetRow& HeaderRow, IPropertyTypeCustomizationUtils& StructCustomizationUtils ) OVERRIDE;
-	virtual void CustomizeChildren( TSharedRef<class IPropertyHandle> StructPropertyHandle, class IDetailChildrenBuilder& StructBuilder, IPropertyTypeCustomizationUtils& StructCustomizationUtils ) OVERRIDE;
-	// End IPropertyTypeCustomization interface
 
+	// IPropertyTypeCustomization interface
+
+	virtual void CustomizeHeader( TSharedRef<IPropertyHandle> StructPropertyHandle, class FDetailWidgetRow& HeaderRow, IPropertyTypeCustomizationUtils& StructCustomizationUtils ) override;
+	virtual void CustomizeChildren( TSharedRef<class IPropertyHandle> StructPropertyHandle, class IDetailChildrenBuilder& StructBuilder, IPropertyTypeCustomizationUtils& StructCustomizationUtils ) override;
+
+public:
+
+	/**
+	 * Creates a new instance.
+	 *
+	 * @return A new struct customization for Keys.
+	 */
 	static TSharedRef<IPropertyTypeCustomization> MakeInstance();
 };
diff --git a/Engine/Source/Editor/DetailCustomizations/Private/RangeStructCustomization.cpp b/Engine/Source/Editor/DetailCustomizations/Private/RangeStructCustomization.cpp
--- a/Engine/Source/Editor/DetailCustomizations/Private/RangeStructCustomization.cpp
+++ b/Engine/Source/Editor/DetailCustomizations/Private/RangeStructCustomization.cpp
@@ -1,9 +1,5 @@
 // Copyright 1998-2014 Epic Games, Inc. All Rights Reserved.
 
-/*=============================================================================
-	RangeStructCustomization.cpp: Implements the FRangeStructCustomization class.
-  =============================================================================*/
-
 #include "DetailCustomizationsPrivatePCH.h"
 #include "RangeStructCustomization.h"
 
@@ -13,11 +9,12 @@
 
 /* Helper traits for getting a metadata property based on the template parameter type
  *****************************************************************************/
+
 namespace
 {
 	template <typename NumericType>
 	struct FGetMetaDataHelper
-	{};
+	{ };
 
 	template <>
 	struct FGetMetaDataHelper<float>
@@ -227,7 +224,6 @@ void FRangeStructCustomization<NumericType>::CustomizeChildren(TSharedRef<IPrope
 }
 
 
-
 /* FRangeStructCustomization callbacks
  *****************************************************************************/
 
@@ -257,6 +253,7 @@ TOptional<NumericType> FRangeStructCustomization<NumericType>::OnGetValue(TWeakP
 	return TOptional<NumericType>();
 }
 
+
 template <typename NumericType>
 void FRangeStructCustomization<NumericType>::OnValueCommitted(NumericType NewValue, ETextCommit::Type CommitType, TWeakPtr<IPropertyHandle> HandleWeakPtr)
 {
@@ -267,6 +264,7 @@ void FRangeStructCustomization<NumericType>::OnValueCommitted(NumericType NewVal
 	}
 }
 
+
 template <typename NumericType>
 void FRangeStructCustomization<NumericType>::OnValueChanged(NumericType NewValue, TWeakPtr<IPropertyHandle> HandleWeakPtr)
 {
@@ -280,6 +278,7 @@ void FRangeStructCustomization<NumericType>::OnValueChanged(NumericType NewValue
 	}
 }
 
+
 template <typename NumericType>
 void FRangeStructCustomization<NumericType>::OnBeginSliderMovement()
 {
@@ -291,6 +290,7 @@ void FRangeStructCustomization<NumericType>::OnBeginSliderMovement()
 	}
 }
 
+
 template <typename NumericType>
 void FRangeStructCustomization<NumericType>::OnEndSliderMovement(NumericType /*NewValue*/)
 {
@@ -302,6 +302,7 @@ void FRangeStructCustomization<NumericType>::OnEndSliderMovement(NumericType /*N
 	}
 }
 
+
 template <typename NumericType>
 bool FRangeStructCustomization<NumericType>::OnQueryIfEnabled(TWeakPtr<IPropertyHandle> HandleWeakPtr) const
 {
@@ -318,6 +319,7 @@ bool FRangeStructCustomization<NumericType>::OnQueryIfEnabled(TWeakPtr<IProperty
 	return false;
 }
 
+
 template <typename NumericType>
 bool FRangeStructCustomization<NumericType>::ShouldAllowSpin() const
 {
@@ -332,6 +334,7 @@ bool FRangeStructCustomization<NumericType>::ShouldAllowSpin() const
 	return false;
 }
 
+
 template <typename NumericType>
 TSharedRef<SWidget> FRangeStructCustomization<NumericType>::OnGenerateComboWidget(TSharedPtr<FString> InComboString)
 {
@@ -361,6 +364,7 @@ TSharedRef<SWidget> FRangeStructCustomization<NumericType>::OnGenerateComboWidge
 		];
 }
 
+
 template <typename NumericType>
 void FRangeStructCustomization<NumericType>::OnComboSelectionChanged(TSharedPtr<FString> InSelectedItem, ESelectInfo::Type SelectInfo, TWeakPtr<IPropertyHandle> HandleWeakPtr)
 {
@@ -378,6 +382,7 @@ void FRangeStructCustomization<NumericType>::OnComboSelectionChanged(TSharedPtr<
 
 /* Only explicitly instantiate the types which are supported
  *****************************************************************************/
+
 template class FRangeStructCustomization<float>;
 template class FRangeStructCustomization<int32>;
 
diff --git a/Engine/Source/Editor/DetailCustomizations/Private/RangeStructCustomization.h b/Engine/Source/Editor/DetailCustomizations/Private/RangeStructCustomization.h
--- a/Engine/Source/Editor/DetailCustomizations/Private/RangeStructCustomization.h
+++ b/Engine/Source/Editor/DetailCustomizations/Private/RangeStructCustomization.h
@@ -1,9 +1,5 @@
 // Copyright 1998-2014 Epic Games, Inc. All Rights Reserved.
 
-/*=============================================================================
-	RangeStructCustomization.h: Declares the FRangeStructCustomization<> class.
-=============================================================================*/
-
 #pragma once
 
 /**
@@ -17,13 +13,12 @@ public:
 	FRangeStructCustomization()
 		: bIsUsingSlider(false) {}
 
-	// Begin IPropertyTypeCustomization interface
-
-	virtual void CustomizeHeader(TSharedRef<IPropertyHandle> StructPropertyHandle, FDetailWidgetRow& HeaderRow, IPropertyTypeCustomizationUtils& StructCustomizationUtils) OVERRIDE;
+public:
 
-	virtual void CustomizeChildren(TSharedRef<IPropertyHandle> StructPropertyHandle, IDetailChildrenBuilder& StructBuilder, IPropertyTypeCustomizationUtils& StructCustomizationUtils) OVERRIDE;
+	// IPropertyTypeCustomization interface
 
-	// End IPropertyTypeCustomization interface
+	virtual void CustomizeHeader(TSharedRef<IPropertyHandle> StructPropertyHandle, FDetailWidgetRow& HeaderRow, IPropertyTypeCustomizationUtils& StructCustomizationUtils) override;
+	virtual void CustomizeChildren(TSharedRef<IPropertyHandle> StructPropertyHandle, IDetailChildrenBuilder& StructBuilder, IPropertyTypeCustomizationUtils& StructCustomizationUtils) override;
 
 public:
 
@@ -41,7 +36,6 @@ protected:
 	 *
 	 * @param	ValueWeakPtr	Handle to the property to get the value from
 	 * @param	TypeWeakPtr		Handle to the property to get the type from
-	 *
 	 * @return	The value or unset if it could not be accessed
 	 */
 	TOptional<NumericType> OnGetValue(TWeakPtr<IPropertyHandle> ValueWeakPtr, TWeakPtr<IPropertyHandle> TypeWeakPtr) const;
@@ -79,15 +73,14 @@ protected:
 	 * Determines if the value is valid from the handle of the range type passed in
 	 *
 	 * @param	HandleWeakPtr	Handle to the property to get the type from
-	 *
 	 * @return	Whether the value is valid or not.
 	 */
 	bool OnQueryIfEnabled(TWeakPtr<IPropertyHandle> HandleWeakPtr) const;
 
 	/**
-	 * Determines if the spinbox is enabled on a numeric value widget
+	 * Determines if the spin box is enabled on a numeric value widget
 	 *
-	 * @return Whether the spinbox should be enabled.
+	 * @return Whether the spin box should be enabled.
 	 */
 	bool ShouldAllowSpin() const;
 
@@ -95,7 +88,6 @@ protected:
 	 * Generates a row of the combo widget
 	 *
 	 * @param	InComboString	Item string to be displayed in the combo item
-	 *
 	 * @return	An SWidget representing the combo row.
 	 */
 	TSharedRef<SWidget> OnGenerateComboWidget(TSharedPtr<FString> InComboString);
diff --git a/Engine/Source/Editor/DetailCustomizations/Private/WorldSettingsDetails.cpp b/Engine/Source/Editor/DetailCustomizations/Private/WorldSettingsDetails.cpp
--- a/Engine/Source/Editor/DetailCustomizations/Private/WorldSettingsDetails.cpp
+++ b/Engine/Source/Editor/DetailCustomizations/Private/WorldSettingsDetails.cpp
@@ -1,9 +1,5 @@
 // Copyright 1998-2014 Epic Games, Inc. All Rights Reserved.
 
-/*=============================================================================
-	WorldSettingsDetails.cpp: Implements the FWorldSettingsDetails class.
-=============================================================================*/
-
 #include "DetailCustomizationsPrivatePCH.h"
 #include "WorldSettingsDetails.h"
 #include "SAssetDropTarget.h"
diff --git a/Engine/Source/Editor/DetailCustomizations/Private/WorldSettingsDetails.h b/Engine/Source/Editor/DetailCustomizations/Private/WorldSettingsDetails.h
--- a/Engine/Source/Editor/DetailCustomizations/Private/WorldSettingsDetails.h
+++ b/Engine/Source/Editor/DetailCustomizations/Private/WorldSettingsDetails.h
@@ -1,9 +1,5 @@
 // Copyright 1998-2014 Epic Games, Inc. All Rights Reserved.
 
-/*=============================================================================
-	WorldSettingsDetails.h: Declares the FWorldSettingsDetails class.
-=============================================================================*/
-
 #pragma once
 
 #include "GameModeInfoCustomizer.h"
@@ -17,12 +13,9 @@ class FWorldSettingsDetails
 {
 public:
 
-	// Begin IDetailCustomization interface
-
-	virtual void CustomizeDetails( IDetailLayoutBuilder& DetailBuilder ) OVERRIDE;
-
-	// End IDetailCustomization interface
+	// IDetailCustomization interface
 
+	virtual void CustomizeDetails( IDetailLayoutBuilder& DetailBuilder ) override;
 
 public:
 
@@ -36,22 +29,21 @@ public:
 		return MakeShareable(new FWorldSettingsDetails);
 	}
 
-
 protected:
 
 	/**
 	 * Customizes an AGameInfo property with the given name.
 	 *
-	 * @param PropertyName - The property to customize.
-	 * @param DetailBuilder - The detail builder.
-	 * @param CategoryBuilder - The category builder
+	 * @param PropertyName The property to customize.
+	 * @param DetailBuilder The detail builder.
+	 * @param CategoryBuilder The category builder
 	 */
 	void CustomizeGameInfoProperty( const FName& PropertyName, IDetailLayoutBuilder& DetailBuilder, IDetailCategoryBuilder& CategoryBuilder );
 
 	/**
 	 * Adds the lightmap customization to the Lightmass section
 	 *
-	 * @param DetailBuilder - The detail builder.
+	 * @param DetailBuilder The detail builder.
 	 */
 	void AddLightmapCustomization( IDetailLayoutBuilder& DetailBuilder );
 
@@ -68,6 +60,7 @@ private:
 	TSharedPtr<FGameModeInfoCustomizer>	GameInfoModeCustomizer;
 };
 
+
 /** Custom struct for each group of arguments in the function editing details */
 class FLightmapCustomNodeBuilder : public IDetailCustomNodeBuilder, public TSharedFromThis<FLightmapCustomNodeBuilder>
 {
@@ -77,15 +70,14 @@ public:
 
 protected:
 
-	// Begin IDetailCustomNodeBuilder interface
-	virtual void SetOnRebuildChildren( FSimpleDelegate InOnRegenerateChildren ) OVERRIDE;
-	virtual void GenerateHeaderRowContent( FDetailWidgetRow& NodeRow ) OVERRIDE;
-	virtual void GenerateChildContent( IDetailChildrenBuilder& ChildrenBuilder ) OVERRIDE;
-	virtual void Tick( float DeltaTime ) OVERRIDE {}
-	virtual bool RequiresTick() const OVERRIDE { return false; }
-	virtual FName GetName() const OVERRIDE { return FName(TEXT("Lightmaps")); }
-	virtual bool InitiallyCollapsed() const OVERRIDE { return false; }
-	// End IDetailCustomNodeBuilder interface
+	// IDetailCustomNodeBuilder interface
+	virtual void SetOnRebuildChildren( FSimpleDelegate InOnRegenerateChildren ) override;
+	virtual void GenerateHeaderRowContent( FDetailWidgetRow& NodeRow ) override;
+	virtual void GenerateChildContent( IDetailChildrenBuilder& ChildrenBuilder ) override;
+	virtual void Tick( float DeltaTime ) override {}
+	virtual bool RequiresTick() const override { return false; }
+	virtual FName GetName() const override { return FName(TEXT("Lightmaps")); }
+	virtual bool InitiallyCollapsed() const override { return false; }
 
 private:
 
@@ -136,4 +128,4 @@ private:
 	TSharedPtr<SListView<TSharedPtr<FLightmapItem>>> LightmapListView;
 	TArray<TSharedPtr<FLightmapItem>> LightmapItems;
 	TSharedPtr<class FAssetThumbnailPool> ThumbnailPool;
-};
\ No newline at end of file
+};
diff --git a/Engine/Source/Editor/EditorSettingsViewer/Private/EditorSettingsViewerModule.cpp b/Engine/Source/Editor/EditorSettingsViewer/Private/EditorSettingsViewerModule.cpp
--- a/Engine/Source/Editor/EditorSettingsViewer/Private/EditorSettingsViewerModule.cpp
+++ b/Engine/Source/Editor/EditorSettingsViewer/Private/EditorSettingsViewerModule.cpp
@@ -1,12 +1,9 @@
 // Copyright 1998-2014 Epic Games, Inc. All Rights Reserved.
 
-/*=============================================================================
-	EditorSettingsViewerModule.cpp: Implements the FEditorSettingsViewerModule class.
-=============================================================================*/
-
 #include "EditorSettingsViewerPrivatePCH.h"
 #include "Tests/AutomationTestSettings.h"
 
+
 #define LOCTEXT_NAMESPACE "FEditorSettingsViewerModule"
 
 static const FName EditorSettingsTabName("EditorSettings");
@@ -21,9 +18,9 @@ class FEditorSettingsViewerModule
 {
 public:
 
-	// Begin ISettingsViewer interface
+	// ISettingsViewer interface
 
-	virtual void ShowSettings( const FName& CategoryName, const FName& SectionName ) OVERRIDE
+	virtual void ShowSettings( const FName& CategoryName, const FName& SectionName ) override
 	{
 		FGlobalTabmanager::Get()->InvokeTab(EditorSettingsTabName);
 		ISettingsEditorModelPtr SettingsEditorModel = SettingsEditorModelPtr.Pin();
@@ -39,13 +36,11 @@ public:
 		}
 	}
 
-	// End ISettingsViewer interface
-
 public:
 
-	// Begin IModuleInterface interface
+	// IModuleInterface interface
 
-	virtual void StartupModule( ) OVERRIDE
+	virtual void StartupModule( ) override
 	{
 		ISettingsModule* SettingsModule = ISettingsModule::Get();
 
@@ -63,19 +58,17 @@ public:
 			.SetMenuType(ETabSpawnerMenuType::Hide);
 	}
 
-	virtual void ShutdownModule( ) OVERRIDE
+	virtual void ShutdownModule( ) override
 	{
 		FGlobalTabmanager::Get()->UnregisterNomadTabSpawner(EditorSettingsTabName);
 		UnregisterSettings();
 	}
 
-	virtual bool SupportsDynamicReloading( ) OVERRIDE
+	virtual bool SupportsDynamicReloading( ) override
 	{
 		return true;
 	}
 
-	// End IModuleInterface interface
-
 protected:
 
 	/**
diff --git a/Engine/Source/Editor/EditorSettingsViewer/Private/EditorSettingsViewerPrivatePCH.h b/Engine/Source/Editor/EditorSettingsViewer/Private/EditorSettingsViewerPrivatePCH.h
--- a/Engine/Source/Editor/EditorSettingsViewer/Private/EditorSettingsViewerPrivatePCH.h
+++ b/Engine/Source/Editor/EditorSettingsViewer/Private/EditorSettingsViewerPrivatePCH.h
@@ -1,9 +1,5 @@
 // Copyright 1998-2014 Epic Games, Inc. All Rights Reserved.
 
-/*=============================================================================
-	EditorSettingsViewerPrivatePCH.h: Pre-compiled header file for the EditorSettingsViewer module.
-=============================================================================*/
-
 #pragma once
 
 
diff --git a/Engine/Source/Editor/EditorStyle/Private/EditorStyleClasses.cpp b/Engine/Source/Editor/EditorStyle/Private/EditorStyleClasses.cpp
--- a/Engine/Source/Editor/EditorStyle/Private/EditorStyleClasses.cpp
+++ b/Engine/Source/Editor/EditorStyle/Private/EditorStyleClasses.cpp
@@ -1,9 +1,5 @@
 // Copyright 1998-2014 Epic Games, Inc. All Rights Reserved.
 
-/*=============================================================================
-	EditorStyleClasses.cpp: Implements the constructors for the various classes.
-=============================================================================*/
-
 #include "EditorStylePrivatePCH.h"
 
 
diff --git a/Engine/Source/Editor/EditorStyle/Private/EditorStyleModule.cpp b/Engine/Source/Editor/EditorStyle/Private/EditorStyleModule.cpp
--- a/Engine/Source/Editor/EditorStyle/Private/EditorStyleModule.cpp
+++ b/Engine/Source/Editor/EditorStyle/Private/EditorStyleModule.cpp
@@ -1,13 +1,8 @@
 // Copyright 1998-2014 Epic Games, Inc. All Rights Reserved.
 
-/*=============================================================================
-	EditorStyleModule.cpp: Implements the FEditorStyleModule class.
-=============================================================================*/
-
 #include "EditorStylePrivatePCH.h"
 
 
-
 /**
  * Implements the Editor style module, loaded by SlateApplication dynamically at startup.
  */
@@ -16,19 +11,19 @@ class FEditorStyleModule
 {
 public:
 
-	// Begin IModuleInterface interface
+	// IEditorStyleModule interface
 
-	virtual void StartupModule( )
+	virtual void StartupModule( ) override
 	{
 		FSlateEditorStyle::Initialize();
 	}
 
-	virtual void ShutdownModule( )
+	virtual void ShutdownModule( ) override
 	{
 		FSlateEditorStyle::Shutdown();
 	}
 
-	virtual TSharedRef<class FSlateStyleSet> CreateEditorStyleInstance( ) const OVERRIDE
+	virtual TSharedRef<class FSlateStyleSet> CreateEditorStyleInstance( ) const override
 	{
 		return FSlateEditorStyle::Create(FSlateEditorStyle::Settings);
 	}
diff --git a/Engine/Source/Editor/EditorStyle/Private/EditorStylePrivatePCH.h b/Engine/Source/Editor/EditorStyle/Private/EditorStylePrivatePCH.h
--- a/Engine/Source/Editor/EditorStyle/Private/EditorStylePrivatePCH.h
+++ b/Engine/Source/Editor/EditorStyle/Private/EditorStylePrivatePCH.h
@@ -1,13 +1,9 @@
 // Copyright 1998-2014 Epic Games, Inc. All Rights Reserved.
 
-/*=============================================================================
-	EditorStylePrivatePCH.h: Pre-compiled header file for the EditorStyle module.
-=============================================================================*/
-
 #pragma once
 
 
-#include "../Public/EditorStyle.h"
+#include "EditorStyle.h"
 
 
 /* Dependencies
diff --git a/Engine/Source/Editor/EditorStyle/Private/SlateEditorStyle.h b/Engine/Source/Editor/EditorStyle/Private/SlateEditorStyle.h
--- a/Engine/Source/Editor/EditorStyle/Private/SlateEditorStyle.h
+++ b/Engine/Source/Editor/EditorStyle/Private/SlateEditorStyle.h
@@ -1,9 +1,5 @@
 // Copyright 1998-2014 Epic Games, Inc. All Rights Reserved.
 
-/*=============================================================================
-	SlateEditorStyle.h: Implements the FSlateEditorStyle class.
-=============================================================================*/
-
 #pragma once
 
 
diff --git a/Engine/Source/Editor/EditorStyle/Public/Classes/EditorStyleSettings.h b/Engine/Source/Editor/EditorStyle/Public/Classes/EditorStyleSettings.h
--- a/Engine/Source/Editor/EditorStyle/Public/Classes/EditorStyleSettings.h
+++ b/Engine/Source/Editor/EditorStyle/Public/Classes/EditorStyleSettings.h
@@ -1,9 +1,5 @@
 // Copyright 1998-2014 Epic Games, Inc. All Rights Reserved.
 
-/*=============================================================================
-	EditorStyleSettings.h: Declares the UEditorStyleSettings class.
-=============================================================================*/
-
 #pragma once
 
 
@@ -90,16 +86,14 @@ public:
 
 protected:
 
-	// Begin UObject overrides
+	// UObject overrides
 
 #if WITH_EDITOR
-	virtual void PostEditChangeProperty( struct FPropertyChangedEvent& PropertyChangedEvent ) OVERRIDE;
+	virtual void PostEditChangeProperty( struct FPropertyChangedEvent& PropertyChangedEvent ) override;
 #endif
 
-	// End UObject overrides
-
 private:
 
 	// Holds an event delegate that is executed when a setting has changed.
 	FSettingChangedEvent SettingChangedEvent;
-};
\ No newline at end of file
+};
diff --git a/Engine/Source/Editor/EditorStyle/Public/EditorStyle.h b/Engine/Source/Editor/EditorStyle/Public/EditorStyle.h
--- a/Engine/Source/Editor/EditorStyle/Public/EditorStyle.h
+++ b/Engine/Source/Editor/EditorStyle/Public/EditorStyle.h
@@ -1,9 +1,5 @@
 // Copyright 1998-2014 Epic Games, Inc. All Rights Reserved.
 
-/*=============================================================================
-	EditorStyle.h: EditorStyle module public header file.
-=============================================================================*/
-
 #pragma once
 
 
diff --git a/Engine/Source/Editor/MainFrame/Private/Frame/MainFrameHandler.h b/Engine/Source/Editor/MainFrame/Private/Frame/MainFrameHandler.h
--- a/Engine/Source/Editor/MainFrame/Private/Frame/MainFrameHandler.h
+++ b/Engine/Source/Editor/MainFrame/Private/Frame/MainFrameHandler.h
@@ -1,9 +1,5 @@
 // Copyright 1998-2014 Epic Games, Inc. All Rights Reserved.
 
-/*=============================================================================
-	MainFrameHandler.h: Declares the FMainFrameHandler class.
-=============================================================================*/
-
 #pragma once
 
 
@@ -58,7 +54,6 @@ public:
 			GConfig->EmptySection(TEXT("EditorLayouts"), *GEditorLayoutIni);
 		}
 
-
 		// Clear the callback for destructionfrom the main tab; otherwise it will re-enter this shutdown function.
 		if ( MainTabPtr.IsValid() )
 		{
diff --git a/Engine/Source/Editor/MainFrame/Private/MainFrameModule.cpp b/Engine/Source/Editor/MainFrame/Private/MainFrameModule.cpp
--- a/Engine/Source/Editor/MainFrame/Private/MainFrameModule.cpp
+++ b/Engine/Source/Editor/MainFrame/Private/MainFrameModule.cpp
@@ -1,9 +1,5 @@
 // Copyright 1998-2014 Epic Games, Inc. All Rights Reserved.
 
-/*=============================================================================
-	MainFrameModule.cpp: Implements the FMainFrameModule class.
-=============================================================================*/
-
 #include "MainFramePrivatePCH.h"
 #include "CompilerResultsLog.h"
 
diff --git a/Engine/Source/Editor/MainFrame/Private/MainFrameModule.h b/Engine/Source/Editor/MainFrame/Private/MainFrameModule.h
--- a/Engine/Source/Editor/MainFrame/Private/MainFrameModule.h
+++ b/Engine/Source/Editor/MainFrame/Private/MainFrameModule.h
@@ -1,9 +1,5 @@
 // Copyright 1998-2014 Epic Games, Inc. All Rights Reserved.
 
-/*=============================================================================
-	MainFrameModule.h: Declares the FMainFrameModule class.
-=============================================================================*/
-
 #pragma once
 
 
@@ -15,14 +11,11 @@ class FMainFrameModule
 {
 public:
 
-	// Begin IMainFrameModule interface
+	// IMainFrameModule interface
 
 	virtual void CreateDefaultMainFrame( const bool bStartImmersivePIE ) OVERRIDE;
-
 	virtual TSharedRef<SWidget> MakeMainMenu( const TSharedPtr<FTabManager>& TabManager, const TSharedRef< FExtender > Extender ) const OVERRIDE;
-
 	virtual TSharedRef<SWidget> MakeMainTabMenu( const TSharedPtr<FTabManager>& TabManager, const TSharedRef< FExtender > Extender ) const OVERRIDE;
-
 	virtual TSharedRef<SWidget> MakeDeveloperTools( ) const OVERRIDE;
 
 	virtual bool IsWindowInitialized( ) const OVERRIDE
@@ -105,14 +98,11 @@ public:
 		return MainFrameSDKNotInstalled.Broadcast(PlatformName, DocLink);
 	}
 
-	// End IMainFrameModule interface
-
 public:
 
-	// Begin IModuleInterface interface
+	// IModuleInterface interface
 
 	virtual void StartupModule( ) OVERRIDE;
-
 	virtual void ShutdownModule( ) OVERRIDE;
 
 	virtual bool SupportsDynamicReloading( ) OVERRIDE
@@ -120,8 +110,6 @@ public:
 		return true; // @todo: Eventually, this should probably not be allowed.
 	}
 
-	// EndIModuleInterface interface
-
 protected:
 
 	/**
diff --git a/Engine/Source/Editor/MainFrame/Private/MainFramePrivatePCH.h b/Engine/Source/Editor/MainFrame/Private/MainFramePrivatePCH.h
--- a/Engine/Source/Editor/MainFrame/Private/MainFramePrivatePCH.h
+++ b/Engine/Source/Editor/MainFrame/Private/MainFramePrivatePCH.h
@@ -1,16 +1,12 @@
 // Copyright 1998-2014 Epic Games, Inc. All Rights Reserved.
 
-/*=============================================================================
-	MainFramePrivatePCH.h: Pre-compiled header file for the MainFrame module.
-=============================================================================*/
-
 #pragma once
 
 
 #include "MainFrame.h"
 
 
-/* Dependencies
+/* Private dependencies
  *****************************************************************************/
 
 #include "AboutScreen.h"
diff --git a/Engine/Source/Editor/MainFrame/Private/Menus/CookContentMenu.h b/Engine/Source/Editor/MainFrame/Private/Menus/CookContentMenu.h
--- a/Engine/Source/Editor/MainFrame/Private/Menus/CookContentMenu.h
+++ b/Engine/Source/Editor/MainFrame/Private/Menus/CookContentMenu.h
@@ -1,13 +1,10 @@
 // Copyright 1998-2014 Epic Games, Inc. All Rights Reserved.
 
-/*=============================================================================
-	PackageProjectMenu.h: Declares the FPackageProjectMenu class.
-=============================================================================*/
-
 #pragma once
 
 #include "ProjectTargetPlatformEditor.h"
 
+
 #define LOCTEXT_NAMESPACE "FCookContentMenu"
 
 
@@ -21,7 +18,7 @@ public:
 	/**
 	 * Creates the menu.
 	 *
-	 * @param MenuBuilder - The builder for the menu that owns this menu.
+	 * @param MenuBuilder The builder for the menu that owns this menu.
 	 */
 	static void MakeMenu( FMenuBuilder& MenuBuilder )
 	{
@@ -74,11 +71,11 @@ public:
 protected:
 
 	/**
-	* Creates the platform menu entries.
-	*
-	* @param MenuBuilder	- The builder for the menu that owns this menu.
-	* @param Platform		- The target platform we allow packaging for
-	*/
+	 * Creates the platform menu entries.
+	 *
+	 * @param MenuBuilder The builder for the menu that owns this menu.
+	 * @param Platform The target platform we allow packaging for
+	 */
 	static void AddPlatformToMenu(FMenuBuilder& MenuBuilder, const PlatformInfo::FPlatformInfo& PlatformInfo)
 	{
 		IProjectTargetPlatformEditorModule& ProjectTargetPlatformEditorModule = FModuleManager::LoadModuleChecked<IProjectTargetPlatformEditorModule>("ProjectTargetPlatformEditor");
@@ -109,14 +106,13 @@ protected:
 			);
 	}
 
-
 	/**
-	* Creates the platform menu entries for a given platforms sub-platforms.
-	* e.g. Windows has multiple sub-platforms - Win32 and Win64
-	*
-	* @param MenuBuilder		- The builder for the menu that owns this menu.
-	* @param SubPlatformInfos	- The Sub-platform information
-	*/
+	 * Creates the platform menu entries for a given platforms sub-platforms.
+	 * e.g. Windows has multiple sub-platforms - Win32 and Win64
+	 *
+	 * @param MenuBuilder The builder for the menu that owns this menu.
+	 * @param SubPlatformInfos The Sub-platform information
+	 */
 	static void AddPlatformSubPlatformsToMenu(FMenuBuilder& MenuBuilder, TArray<const PlatformInfo::FPlatformInfo*> SubPlatformInfos)
 	{
 		for (const PlatformInfo::FPlatformInfo* SubPlatformInfo : SubPlatformInfos)
diff --git a/Engine/Source/Editor/MainFrame/Private/Menus/MainMenu.cpp b/Engine/Source/Editor/MainFrame/Private/Menus/MainMenu.cpp
--- a/Engine/Source/Editor/MainFrame/Private/Menus/MainMenu.cpp
+++ b/Engine/Source/Editor/MainFrame/Private/Menus/MainMenu.cpp
@@ -1,9 +1,5 @@
 // Copyright 1998-2014 Epic Games, Inc. All Rights Reserved.
 
-/*=============================================================================
-	MainMenu.cpp: Implements the FMainMenu class.
-=============================================================================*/
-
 #include "MainFramePrivatePCH.h"
 
 #include "../../WorkspaceMenuStructure/Public/WorkspaceMenuStructureModule.h"
@@ -16,6 +12,7 @@
 
 #define LOCTEXT_NAMESPACE "MainFileMenu"
 
+
 void FMainMenu::FillFileMenu( FMenuBuilder& MenuBuilder, const TSharedRef<FExtender> Extender )
 {
 	MenuBuilder.BeginSection("FileLoadAndSave", LOCTEXT("LoadSandSaveHeading", "Load and Save"));
diff --git a/Engine/Source/Editor/MainFrame/Private/Menus/MainMenu.h b/Engine/Source/Editor/MainFrame/Private/Menus/MainMenu.h
--- a/Engine/Source/Editor/MainFrame/Private/Menus/MainMenu.h
+++ b/Engine/Source/Editor/MainFrame/Private/Menus/MainMenu.h
@@ -1,9 +1,5 @@
 // Copyright 1998-2014 Epic Games, Inc. All Rights Reserved.
 
-/*=============================================================================
-	MainMenu.h: Declares the FMainMenu class.
-=============================================================================*/
-
 #pragma once
 
 
@@ -17,9 +13,8 @@ public:
 	/**
 	 * Static: Creates a widget for the main menu bar.
 	 *
-	 * @param TabManager - Create the workspace menu based on this tab manager.
-	 * @param Extender - Extensibility support for the menu.
-	 *
+	 * @param TabManager Create the workspace menu based on this tab manager.
+	 * @param Extender Extensibility support for the menu.
 	 * @return New widget
 	 */
 	static TSharedRef<SWidget> MakeMainMenu( const TSharedPtr<FTabManager>& TabManager, const TSharedRef<FExtender> Extender );
@@ -28,47 +23,45 @@ public:
 	 * Static: Creates a widget for the main tab's menu bar.  This is just like the main menu bar, but also includes.
 	 * some "project level" menu items that we don't want propagated to most normal menus.
 	 *
-	 * @param TabManager - Create the workspace menu based on this tab manager.
-	 * @param UserExtender - Extensibility support for the menu.
-	 *
+	 * @param TabManager Create the workspace menu based on this tab manager.
+	 * @param UserExtender Extensibility support for the menu.
 	 * @return New widget
 	 */
 	static TSharedRef<SWidget> MakeMainTabMenu( const TSharedPtr<FTabManager>& TabManager, const TSharedRef<FExtender> UserExtender );
 
-
 protected:
 
 	/**
 	 * Called to fill the file menu's content.
 	 *
-	 * @param MenuBuilder (In, Out) - The multi-box builder that should be filled with content for this pull-down menu.
-	 * @param Extender - Extensibility support for this menu.
+	 * @param MenuBuilder The multi-box builder that should be filled with content for this pull-down menu.
+	 * @param Extender Extensibility support for this menu.
 	 */
 	static void FillFileMenu( FMenuBuilder& MenuBuilder, const TSharedRef<FExtender> Extender );
 
 	/**
 	 * Called to fill the edit menu's content.
 	 *
-	 * @param MenuBuilder (In, Out)  - The multi-box builder that should be filled with content for this pull-down menu.
-	 * @param Extender - Extensibility support for this menu.
-	 * @param TabManager - A Tab Manager from which to populate tab spawner menu items.
+	 * @param MenuBuilder The multi-box builder that should be filled with content for this pull-down menu.
+	 * @param Extender Extensibility support for this menu.
+	 * @param TabManager A Tab Manager from which to populate tab spawner menu items.
 	 */
 	static void FillEditMenu( FMenuBuilder& MenuBuilder, const TSharedRef<FExtender> Extender, const TSharedPtr<FTabManager> TabManager );
 
 	/**
 	 * Called to fill the app menu's content.
 	 *
-	 * @param MenuBuilder (In, Out) - The multi-box builder that should be filled with content for this pull-down menu.
-	 * @param Extender - Extensibility support for this menu.
-	 * @param TabManager - A Tab Manager from which to populate tab spawner menu items.
+	 * @param MenuBuilder The multi-box builder that should be filled with content for this pull-down menu.
+	 * @param Extender Extensibility support for this menu.
+	 * @param TabManager A Tab Manager from which to populate tab spawner menu items.
 	 */
 	static void FillWindowMenu( FMenuBuilder& MenuBuilder, const TSharedRef<FExtender> Extender, const TSharedPtr<FTabManager> TabManager );
 
 	/**
 	 * Called to fill the help menu's content.
 	 *
-	 * @param MenuBuilder (In, Out) - The multi-box builder that should be filled with content for this pull-down menu.
-	 * @param Extender - Extensibility support for this menu.
+	 * @param MenuBuilder The multi-box builder that should be filled with content for this pull-down menu.
+	 * @param Extender Extensibility support for this menu.
 	 */
 	static void FillHelpMenu( FMenuBuilder& MenuBuilder, const TSharedRef<FExtender> Extender );
 };
diff --git a/Engine/Source/Editor/MainFrame/Private/Menus/PackageProjectMenu.h b/Engine/Source/Editor/MainFrame/Private/Menus/PackageProjectMenu.h
--- a/Engine/Source/Editor/MainFrame/Private/Menus/PackageProjectMenu.h
+++ b/Engine/Source/Editor/MainFrame/Private/Menus/PackageProjectMenu.h
@@ -1,9 +1,5 @@
 // Copyright 1998-2014 Epic Games, Inc. All Rights Reserved.
 
-/*=============================================================================
-	PackageProjectMenu.h: Declares the FPackageProjectMenu class.
-=============================================================================*/
-
 #pragma once
 
 #include "ProjectTargetPlatformEditor.h"
@@ -21,7 +17,7 @@ public:
 	/**
 	 * Creates the menu.
 	 *
-	 * @param MenuBuilder - The builder for the menu that owns this menu.
+	 * @param MenuBuilder The builder for the menu that owns this menu.
 	 */
 	static void MakeMenu( FMenuBuilder& MenuBuilder )
 	{
@@ -86,11 +82,11 @@ public:
 protected:
 
 	/**
-	* Creates the platform menu entries.
-	*
-	* @param MenuBuilder	- The builder for the menu that owns this menu.
-	* @param Platform		- The target platform we allow packaging for
-	*/
+	 * Creates the platform menu entries.
+	 *
+	 * @param MenuBuilder The builder for the menu that owns this menu.
+	 * @param Platform The target platform we allow packaging for
+	 */
 	static void AddPlatformToMenu(FMenuBuilder& MenuBuilder, const PlatformInfo::FPlatformInfo& PlatformInfo)
 	{
 		IProjectTargetPlatformEditorModule& ProjectTargetPlatformEditorModule = FModuleManager::LoadModuleChecked<IProjectTargetPlatformEditorModule>("ProjectTargetPlatformEditor");
@@ -123,14 +119,13 @@ protected:
 			);
 	}
 
-
 	/**
-	* Creates the platform menu entries for a given platforms sub-platforms.
-	* e.g. Windows has multiple sub-platforms - Win32 and Win64
-	*
-	* @param MenuBuilder		- The builder for the menu that owns this menu.
-	* @param SubPlatformInfos	- The Sub-platform information
-	*/
+	 * Creates the platform menu entries for a given platforms sub-platforms.
+	 * e.g. Windows has multiple sub-platforms - Win32 and Win64
+	 *
+	 * @param MenuBuilderThe builder for the menu that owns this menu.
+	 * @param SubPlatformInfos The Sub-platform information
+	 */
 	static void AddPlatformSubPlatformsToMenu(FMenuBuilder& MenuBuilder, TArray<const PlatformInfo::FPlatformInfo*> SubPlatformInfos)
 	{
 		for (const PlatformInfo::FPlatformInfo* SubPlatformInfo : SubPlatformInfos)
@@ -140,10 +135,10 @@ protected:
 	}
 
 	/**
-	* Creates a build configuration sub-menu.
-	*
-	* @param MenuBuilder - The builder for the menu that owns this menu.
-	*/
+	 * Creates a build configuration sub-menu.
+	 *
+	 * @param MenuBuilder The builder for the menu that owns this menu.
+	 */
 	static void MakeBuildConfigurationsMenu(FMenuBuilder& MenuBuilder)
 	{
 		// Only show the debug game option if the game has source code. 
@@ -195,4 +190,4 @@ protected:
 };
 
 
-#undef LOCTEXT_NAMESPACE
\ No newline at end of file
+#undef LOCTEXT_NAMESPACE
diff --git a/Engine/Source/Editor/MainFrame/Private/Menus/RecentProjectsMenu.h b/Engine/Source/Editor/MainFrame/Private/Menus/RecentProjectsMenu.h
--- a/Engine/Source/Editor/MainFrame/Private/Menus/RecentProjectsMenu.h
+++ b/Engine/Source/Editor/MainFrame/Private/Menus/RecentProjectsMenu.h
@@ -1,9 +1,5 @@
 // Copyright 1998-2014 Epic Games, Inc. All Rights Reserved.
 
-/*=============================================================================
-	RecentProjectsMenu.h: Declares the FRecentProjectsMenu class.
-=============================================================================*/
-
 #pragma once
 
 
@@ -20,7 +16,7 @@ public:
 	/**
 	 * Creates the menu.
 	 *
-	 * @param MenuBuilder - The builder for the menu that owns this menu.
+	 * @param MenuBuilder The builder for the menu that owns this menu.
 	 */
 	static void MakeMenu( FMenuBuilder& MenuBuilder )
 	{
diff --git a/Engine/Source/Editor/MainFrame/Private/Menus/SettingsMenu.h b/Engine/Source/Editor/MainFrame/Private/Menus/SettingsMenu.h
--- a/Engine/Source/Editor/MainFrame/Private/Menus/SettingsMenu.h
+++ b/Engine/Source/Editor/MainFrame/Private/Menus/SettingsMenu.h
@@ -1,9 +1,5 @@
 // Copyright 1998-2014 Epic Games, Inc. All Rights Reserved.
 
-/*=============================================================================
-	SettingsMenu.h: Declares the FSettingsMenu class.
-=============================================================================*/
-
 #pragma once
 
 
@@ -20,8 +16,8 @@ public:
 	/**
 	 * Creates the menu.
 	 *
-	 * @param MenuBuilder - The builder for the menu that owns this menu.
-	 * @param SettingsContainerName - The name of the settings container to create the menu for.
+	 * @param MenuBuilder The builder for the menu that owns this menu.
+	 * @param SettingsContainerName The name of the settings container to create the menu for.
 	 */
 	static void MakeMenu( FMenuBuilder& MenuBuilder, FName SettingsContainerName )
 	{
@@ -75,9 +71,9 @@ public:
 	/**
 	 * Opens the settings tab with the specified settings section.
 	 *
-	 * @param ContainerName - The name of the settings container to open.
-	 * @param CategoryName - The name of the settings category that contains the section.
-	 * @param SectionName - The name of the settings section to select.
+	 * @param ContainerName The name of the settings container to open.
+	 * @param CategoryName The name of the settings category that contains the section.
+	 * @param SectionName The name of the settings section to select.
 	 */
 	static void OpenSettings( FName ContainerName, FName CategoryName, FName SectionName )
 	{
diff --git a/Engine/Source/Editor/MainFrame/Private/Menus/TranslationEditorMenu.h b/Engine/Source/Editor/MainFrame/Private/Menus/TranslationEditorMenu.h
--- a/Engine/Source/Editor/MainFrame/Private/Menus/TranslationEditorMenu.h
+++ b/Engine/Source/Editor/MainFrame/Private/Menus/TranslationEditorMenu.h
@@ -1,9 +1,5 @@
 // Copyright 1998-2014 Epic Games, Inc. All Rights Reserved.
 
-/*=============================================================================
-	PackageProjectMenu.h: Declares the FPackageProjectMenu class.
-=============================================================================*/
-
 #pragma once
 
 
@@ -20,7 +16,7 @@ public:
 	/**
 	 * Creates the main frame translation editor sub menu.
 	 *
-	 * @param MenuBuilder - The builder for the menu that owns this menu.
+	 * @param MenuBuilder The builder for the menu that owns this menu.
 	 */
 	static void MakeMainFrameTranslationEditorSubMenu( FMenuBuilder& MenuBuilder )
 	{
@@ -118,7 +114,7 @@ public:
 	/**
 	 * Creates the main frame translation editor sub menu for editor translations.
 	 *
-	 * @param MenuBuilder - The builder for the menu that owns this menu.
+	 * @param MenuBuilder The builder for the menu that owns this menu.
 	 */
 	static void MakeMainFrameTranslationEditorSubMenuEditorProject( FMenuBuilder& MenuBuilder, const FString ManifestFileName )
 	{
diff --git a/Engine/Source/Editor/MainFrame/Public/Interfaces/IMainFrameModule.h b/Engine/Source/Editor/MainFrame/Public/Interfaces/IMainFrameModule.h
--- a/Engine/Source/Editor/MainFrame/Public/Interfaces/IMainFrameModule.h
+++ b/Engine/Source/Editor/MainFrame/Public/Interfaces/IMainFrameModule.h
@@ -1,9 +1,5 @@
 // Copyright 1998-2014 Epic Games, Inc. All Rights Reserved.
 
-/*=============================================================================
-	IMainFrameModule.h: Declares the IMainFrameModule interface.
-=============================================================================*/
-
 #pragma once
 
 
@@ -142,8 +138,7 @@ public:
 	/**
 	 * Gets a reference to the search module instance.
 	 *
-	 * @todo gmp - better implementation using dependency injection.
-	 *
+	 * @todo gmp: better implementation using dependency injection.
 	 * @return A reference to the MainFrame module.
 	 */
 	static IMainFrameModule& Get( )
@@ -151,7 +146,6 @@ public:
 		return FModuleManager::LoadModuleChecked<IMainFrameModule>("MainFrame");
 	}
 
-
 public:
 
 	/**
diff --git a/Engine/Source/Editor/MainFrame/Public/MainFrame.h b/Engine/Source/Editor/MainFrame/Public/MainFrame.h
--- a/Engine/Source/Editor/MainFrame/Public/MainFrame.h
+++ b/Engine/Source/Editor/MainFrame/Public/MainFrame.h
@@ -1,20 +1,16 @@
 // Copyright 1998-2014 Epic Games, Inc. All Rights Reserved.
 
-/*=============================================================================
-	MainFrame.h: Public header file for the MainFrame module.
-=============================================================================*/
-
 #pragma once
 
 
-/* Dependencies
+/* Public Dependencies
  *****************************************************************************/
 
 #include "ModuleInterface.h"
 #include "Slate.h"
 
 
-/* Interfaces
+/* Public Includes
  *****************************************************************************/
 
 #include "IMainFrameModule.h"
diff --git a/Engine/Source/Editor/Matinee/Classes/MatineeTrackAnimControlHelper.h b/Engine/Source/Editor/Matinee/Classes/MatineeTrackAnimControlHelper.h
--- a/Engine/Source/Editor/Matinee/Classes/MatineeTrackAnimControlHelper.h
+++ b/Engine/Source/Editor/Matinee/Classes/MatineeTrackAnimControlHelper.h
@@ -2,23 +2,23 @@
 
 #pragma once
 
-
 #include "MatineeTrackAnimControlHelper.generated.h"
 
+
 UCLASS()
 class UMatineeTrackAnimControlHelper : public UInterpTrackHelper
 {
 	GENERATED_UCLASS_BODY()
 
-
-	// Begin UInterpTrackHelper Interface
-	virtual	bool PreCreateTrack( UInterpGroup* Group, const UInterpTrack *TrackDef, bool bDuplicatingTrack, bool bAllowPrompts ) const OVERRIDE;
-	virtual void  PostCreateTrack( UInterpTrack *Track, bool bDuplicatingTrack, int32 TrackIndex ) const OVERRIDE;
-	virtual	bool PreCreateKeyframe( UInterpTrack *Track, float KeyTime ) const OVERRIDE;
-	virtual void  PostCreateKeyframe( UInterpTrack *Track, int32 KeyIndex ) const OVERRIDE;
-	// End UInterpTrackHelper Interface
-
 	void OnAddKeyTextEntry(const class FAssetData& AssetData, IMatineeBase* Matinee, UInterpTrack* Track);
 	void OnCreateTrackTextEntry(const FString& ChosenText, TSharedRef<SWindow> Window, FString* OutputString);
-};
 
+public:
+
+	// UInterpTrackHelper interface
+
+	virtual	bool PreCreateTrack( UInterpGroup* Group, const UInterpTrack *TrackDef, bool bDuplicatingTrack, bool bAllowPrompts ) const override;
+	virtual void  PostCreateTrack( UInterpTrack *Track, bool bDuplicatingTrack, int32 TrackIndex ) const override;
+	virtual	bool PreCreateKeyframe( UInterpTrack *Track, float KeyTime ) const override;
+	virtual void  PostCreateKeyframe( UInterpTrack *Track, int32 KeyIndex ) const override;
+};
