fix for UFE crash on shutdown
TTP338031
#ue4
#ufe

diff --git a/Engine/Source/Developer/TargetDeviceServices/Private/TargetDeviceServiceManager.cpp b/Engine/Source/Developer/TargetDeviceServices/Private/TargetDeviceServiceManager.cpp
--- a/Engine/Source/Developer/TargetDeviceServices/Private/TargetDeviceServiceManager.cpp
+++ b/Engine/Source/Developer/TargetDeviceServices/Private/TargetDeviceServiceManager.cpp
@@ -291,14 +291,18 @@ void FTargetDeviceServiceManager::SaveSettings( )
 
 void FTargetDeviceServiceManager::ShutdownTargetPlatforms( )
 {
-	TArray<ITargetPlatform*> Platforms = GetTargetPlatformManager()->GetTargetPlatforms();
-
-	for (int32 PlatformIndex = 0; PlatformIndex < Platforms.Num(); ++PlatformIndex)
+	ITargetPlatformManagerModule* Module = FModuleManager::GetModulePtr<ITargetPlatformManagerModule>("TargetPlatform");
+	if (Module)
 	{
-		// set up target platform callbacks
-		ITargetPlatform* Platform = Platforms[PlatformIndex];
-		Platform->OnDeviceDiscovered().RemoveAll(this);
-		Platform->OnDeviceLost().RemoveAll(this);
+		TArray<ITargetPlatform*> Platforms = Module->GetTargetPlatforms();
+
+		for (int32 PlatformIndex = 0; PlatformIndex < Platforms.Num(); ++PlatformIndex)
+		{
+			// set up target platform callbacks
+			ITargetPlatform* Platform = Platforms[PlatformIndex];
+			Platform->OnDeviceDiscovered().RemoveAll(this);
+			Platform->OnDeviceLost().RemoveAll(this);
+		}
 	}
 }
 
