Target Platform API added to provide more information about variants and variant support.

diff --git a/Engine/Source/Developer/ProjectLauncher/Private/Widgets/SProjectLauncherLaunchTaskSettings.h b/Engine/Source/Developer/ProjectLauncher/Private/Widgets/SProjectLauncherLaunchTaskSettings.h
new file mode 100644
--- /dev/null
+++ b/Engine/Source/Developer/ProjectLauncher/Private/Widgets/SProjectLauncherLaunchTaskSettings.h
@@ -0,0 +1,38 @@
+// Copyright 1998-2014 Epic Games, Inc. All Rights Reserved.
+
+#pragma once
+
+
+/**
+ * Implements the launcher settings widget.
+ */
+class SProjectLauncherLaunchTaskSettings
+	: public SCompoundWidget
+{
+public:
+
+	SLATE_BEGIN_ARGS(SProjectLauncherLaunchTaskSettings) { }
+	SLATE_END_ARGS()
+
+public:
+
+	/**
+	 * Destructor.
+	 */
+	~SProjectLauncherLaunchTaskSettings( );
+
+public:
+
+	/**
+	 * Constructs the widget.
+	 *
+	 * @param InArgs The Slate argument list.
+	 * @param InModel The data model.
+	 */
+	void Construct( const FArguments& InArgs, const FProjectLauncherModelRef& InModel );
+
+private:
+
+	// Holds a pointer to the data model.
+	FProjectLauncherModelPtr Model;
+};
diff --git a/Engine/Source/Developer/ProjectLauncher/Private/Widgets/SProjectLauncherPreviewPage.h b/Engine/Source/Developer/ProjectLauncher/Private/Widgets/SProjectLauncherPreviewPage.h
new file mode 100644
--- /dev/null
+++ b/Engine/Source/Developer/ProjectLauncher/Private/Widgets/SProjectLauncherPreviewPage.h
@@ -0,0 +1,107 @@
+// Copyright 1998-2014 Epic Games, Inc. All Rights Reserved.
+
+#pragma once
+
+
+/**
+ * Implements the launcher's preview page.
+ */
+class SProjectLauncherPreviewPage
+	: public SCompoundWidget
+{
+public:
+
+	SLATE_BEGIN_ARGS(SProjectLauncherPreviewPage) { }
+	SLATE_END_ARGS()
+
+public:
+
+	/**
+	 * Constructs the widget.
+	 *
+	 * @param InArgs The Slate argument list.
+	 * @param InModel The data model.
+	 */
+	void Construct(	const FArguments& InArgs, const FProjectLauncherModelRef& InModel );
+
+public:
+
+	virtual void Tick( const FGeometry& AllottedGeometry, const double InCurrentTime, const float InDeltaTime ) override;
+
+protected:
+
+	/**
+	 * Refreshes the list of device proxies.
+	 */
+	void RefreshDeviceProxyList( );
+
+private:
+
+	// Callback for getting the text in the 'Build Configuration' text block.
+	FText HandleBuildConfigurationTextBlockText( ) const;
+
+	// Callback for getting the list of platforms to build for.
+	FText HandleBuildPlatformsTextBlockText() const;
+
+	// Callback for getting the text in the 'Command Line' text block.
+	FText HandleCommandLineTextBlockText() const;
+
+	// Callback for getting the text in the 'Cooked Cultures' text block.
+	FText HandleCookedCulturesTextBlockText() const;
+
+	// Callback for getting the text in the 'Cooked Maps' text block.
+	FText HandleCookedMapsTextBlockText() const;
+
+	// Callback for getting the text in the 'Cooker Options' text block.
+	FText HandleCookerOptionsTextBlockText() const;
+
+	// Callback for getting the visibility of the specified cook summary box.
+	EVisibility HandleCookSummaryBoxVisibility( ELauncherProfileCookModes::Type CookMode ) const;
+
+	// Callback for getting the visibility of the deploy details box.
+	EVisibility HandleDeployDetailsBoxVisibility( ) const;
+
+	// Callback for getting the visibility of the specified deploy summary box.
+	EVisibility HandleDeploySummaryBoxVisibility( ELauncherProfileDeploymentModes::Type DeploymentMode ) const;
+
+	// Callback for generating a row for the device list.
+	TSharedRef<ITableRow> HandleDeviceProxyListViewGenerateRow( ITargetDeviceProxyPtr InItem, const TSharedRef<STableViewBase>& OwnerTable );
+
+	// Callback for getting the text in the 'Initial Culture' text block.
+	FString HandleInitialCultureTextBlockText( ) const;
+
+	// Callback for getting the text in the 'Initial Map' text block.
+	FString HandleInitialMapTextBlockText( ) const;
+
+	// Callback for getting the visibility of the specified cook summary box.
+	EVisibility HandleLaunchSummaryBoxVisibility( ELauncherProfileLaunchModes::Type LaunchMode ) const;
+
+	// Callback for getting the text in the 'VSync' text block.
+	FString HandleLaunchVsyncTextBlockText( ) const;
+
+	// Callback for getting the visibility of the specified cook summary box.
+	EVisibility HandlePackageSummaryBoxVisibility( ELauncherProfilePackagingModes::Type PackagingMode ) const;
+
+	// Callback for getting the text in the 'Game' text block.
+	FString HandleProjectTextBlockText( ) const;
+
+	// Callback for getting the name of the selected device group.
+	FString HandleSelectedDeviceGroupTextBlockText( ) const;
+
+	// Callback for getting the name of the selected profile.
+	FString HandleSelectedProfileTextBlockText( ) const;
+
+	// Callback for determining the visibility of a validation error icon.
+	EVisibility HandleValidationErrorIconVisibility( ELauncherProfileValidationErrors::Type Error ) const;
+
+private:
+
+	// Holds the list of available device proxies.
+	TArray<ITargetDeviceProxyPtr> DeviceProxyList;
+
+	// Holds the device proxy list view.
+	TSharedPtr<SListView<ITargetDeviceProxyPtr> > DeviceProxyListView;
+
+	// Holds a pointer to the data model.
+	FProjectLauncherModelPtr Model;
+};
diff --git a/Engine/Source/Developer/ProjectLauncher/Private/Widgets/SProjectLauncherSettings.h b/Engine/Source/Developer/ProjectLauncher/Private/Widgets/SProjectLauncherSettings.h
new file mode 100644
--- /dev/null
+++ b/Engine/Source/Developer/ProjectLauncher/Private/Widgets/SProjectLauncherSettings.h
@@ -0,0 +1,134 @@
+// Copyright 1998-2014 Epic Games, Inc. All Rights Reserved.
+
+/*=============================================================================
+	SProjectLauncherSettings.h: Declares the SProjectLauncherCookPage class.
+=============================================================================*/
+
+#pragma once
+
+
+namespace ELauncherWizardPages
+{
+	/**
+	 * Enumerates the session launcher wizard pages.
+	 */
+	enum Type
+	{
+		/** The 'Build' page. */
+		BuildPage,
+
+		/** The 'Cook' page. */
+		CookPage,
+
+		/** The 'Package' page. */
+		PackagePage,
+
+		/** The 'Deploy' page. */
+		DeployPage,
+
+		/** The 'Launch' page. */
+		LaunchPage,
+
+		/** The 'Preview' page. */
+		PreviewPage
+	};
+}
+
+
+/**
+ * Implements the launcher settings widget.
+ */
+class SProjectLauncherSettings
+	: public SCompoundWidget
+{
+public:
+
+	SLATE_BEGIN_ARGS(SProjectLauncherSettings) { }
+	SLATE_EVENT(FOnClicked, OnCloseClicked)
+	SLATE_EVENT(FOnProfileRun, OnDeleteClicked)
+	SLATE_END_ARGS()
+
+public:
+
+	/**
+	* Default constructor.
+	*/
+	SProjectLauncherSettings()
+		: CommandList(new FUICommandList())
+	{ }
+
+	/**
+	 * Destructor.
+	 */
+	~SProjectLauncherSettings();
+
+public:
+
+	/**
+	 * Constructs the widget.
+	 *
+	 * @param InArgs The Slate argument list.
+	 * @param InModel The data model.
+	 */
+	void Construct(const FArguments& InArgs, const FProjectLauncherModelRef& InModel);
+
+	void EnterEditMode();
+
+private:
+
+	// Create the menu command list
+	void CreateCommands();
+
+	// Make a Toolbar using the command list
+	TSharedRef<SWidget> MakeToolbar(const TSharedRef<FUICommandList>& InCommandList);
+
+private:
+
+	// Callback for getting the launch profile we are currently editing.
+	ILauncherProfilePtr GetLaunchProfile()const;
+
+	// Callback for getting the visibility of the 'Select Profile' text block.
+	EVisibility HandleSelectProfileTextBlockVisibility() const;
+
+	// Callback for getting the visibility of the settings scroll box.
+	EVisibility HandleSettingsScrollBoxVisibility() const;
+
+	// Callback to get the name of the launch profile.
+	FText OnGetNameText() const;
+
+	// Callback to set the name of the launch profile.
+	void OnNameTextCommitted(const FText& NewText, ETextCommit::Type InTextCommit);
+
+	// Callback to get the description of the launch profile.
+	FText OnGetDescriptionText() const;
+
+	// Callback to set the description of the launch profile.
+	void OnDescriptionTextCommitted(const FText& NewText, ETextCommit::Type InTextCommit);
+
+	// Callbacks for executing the 'Close' action.
+	void HandleCloseActionExecute();
+	bool HandleCloseActionIsChecked() const;
+	bool HandleCloseActionCanExecute() const;
+
+	// Callbacks for executing the 'Delete' action.
+	void HandleDeleteActionExecute();
+	bool HandleDeleteActionIsChecked() const;
+	bool HandleDeleteActionCanExecute() const;
+
+private:
+
+	// Holds a pointer to the data model.
+	FProjectLauncherModelPtr Model;
+
+	// Holds the list of UI commands for the profile settings.
+	TSharedRef<FUICommandList> CommandList;
+
+	// Holds a delegate to be invoked when this panel is closed
+	FOnClicked OnCloseClicked;
+
+	// Holds a delegate to be invoked when this profile is deleted
+	FOnProfileRun OnDeleteClicked;
+
+	// Hold a pointer to the launch profile name edit box.
+	TSharedPtr<SInlineEditableTextBlock> NameEditBox;
+};
diff --git a/Engine/Source/Developer/ProjectLauncher/Private/Widgets/SProjectLauncherTaskListRow.h b/Engine/Source/Developer/ProjectLauncher/Private/Widgets/SProjectLauncherTaskListRow.h
new file mode 100644
--- /dev/null
+++ b/Engine/Source/Developer/ProjectLauncher/Private/Widgets/SProjectLauncherTaskListRow.h
@@ -0,0 +1,246 @@
+// Copyright 1998-2014 Epic Games, Inc. All Rights Reserved.
+
+#pragma once
+
+
+#define LOCTEXT_NAMESPACE "SProjectLauncherTaskListRow"
+
+
+/**
+ * Implements a row widget for the launcher's task list.
+ */
+class SProjectLauncherTaskListRow
+	: public SMultiColumnTableRow<ILauncherTaskPtr>
+{
+public:
+
+	SLATE_BEGIN_ARGS(SProjectLauncherTaskListRow) { }
+		SLATE_ARGUMENT(TSharedPtr<STableViewBase>, OwnerTableView)
+		SLATE_ARGUMENT(ILauncherTaskPtr, Task)
+	SLATE_END_ARGS()
+
+public:
+
+	/**
+	 * Constructs the widget.
+	 *
+	 * @param InArgs The construction arguments.
+	 * @param InDeviceProxyManager The device proxy manager to use.
+	 */
+	void Construct( const FArguments& InArgs )
+	{
+		Task = InArgs._Task;
+
+		SMultiColumnTableRow<ILauncherTaskPtr>::Construct(FSuperRowType::FArguments(), InArgs._OwnerTableView.ToSharedRef());
+	}
+
+public:
+
+	/**
+	 * Generates the widget for the specified column.
+	 *
+	 * @param ColumnName The name of the column to generate the widget for.
+	 * @return The widget.
+	 */
+	BEGIN_SLATE_FUNCTION_BUILD_OPTIMIZATION
+	virtual TSharedRef<SWidget> GenerateWidgetForColumn( const FName& ColumnName ) override
+	{
+		if (ColumnName == "Duration")
+		{
+			return SNew(SBox)
+				.Padding(FMargin(4.0, 0.0))
+				.VAlign((VAlign_Center))
+				[
+					SNew(STextBlock)
+						.Text(this, &SProjectLauncherTaskListRow::HandleDurationText)
+				];
+		}
+		else if (ColumnName == "Icon")
+		{
+			return SNew(SOverlay)
+
+			+ SOverlay::Slot()
+				.HAlign(HAlign_Center)
+				.VAlign((VAlign_Center))
+				[
+					SNew(SThrobber)
+						.Animate(SThrobber::VerticalAndOpacity)
+						.NumPieces(1)
+						.Visibility(this, &SProjectLauncherTaskListRow::HandleThrobberVisibility)
+				]
+
+			+ SOverlay::Slot()
+				.HAlign(HAlign_Center)
+				.VAlign((VAlign_Center))
+				[
+					SNew(SImage)
+						.ColorAndOpacity(this, &SProjectLauncherTaskListRow::HandleIconColorAndOpacity)
+						.Image(this, &SProjectLauncherTaskListRow::HandleIconImage)
+				];
+		}
+		else if (ColumnName == "Status")
+		{
+			return SNew(SBox)
+				.Padding(FMargin(4.0, 0.0))
+				.VAlign((VAlign_Center))
+				[
+					SNew(STextBlock)
+						.Text(this, &SProjectLauncherTaskListRow::HandleStatusText)
+				];
+		}
+		else if (ColumnName == "Task")
+		{
+			ILauncherTaskPtr TaskPtr = Task.Pin();
+
+			if (TaskPtr.IsValid())
+			{
+				return SNew(SBox)
+					.Padding(FMargin(4.0, 0.0))
+					.VAlign((VAlign_Center))
+					[
+						SNew(STextBlock)
+							.Text(TaskPtr->GetDesc())
+					];
+			}
+		}
+
+		return SNullWidget::NullWidget;
+	}
+	END_SLATE_FUNCTION_BUILD_OPTIMIZATION
+
+private:
+
+	// Callback for getting the duration of the task.
+	FString HandleDurationText( ) const
+	{
+		ILauncherTaskPtr TaskPtr = Task.Pin();
+
+		if (TaskPtr.IsValid())
+		{
+			if ((TaskPtr->GetStatus() != ELauncherTaskStatus::Pending) &&
+				(TaskPtr->GetStatus() != ELauncherTaskStatus::Canceled))
+			{
+				return TaskPtr->GetDuration().ToString(TEXT("%h:%m:%s"));
+			}
+		}
+
+		return FString();
+	}
+
+	// Callback for getting the color and opacity of the status icon.
+	FSlateColor HandleIconColorAndOpacity( ) const
+	{
+		ILauncherTaskPtr TaskPtr = Task.Pin();
+
+		if (TaskPtr.IsValid())
+		{
+			if (TaskPtr->GetStatus() == ELauncherTaskStatus::Canceled)
+			{
+				return FLinearColor::Yellow;
+			}
+
+			if (TaskPtr->GetStatus() == ELauncherTaskStatus::Completed)
+			{
+				return FLinearColor::Green;
+			}
+
+			if (TaskPtr->GetStatus() == ELauncherTaskStatus::Failed)
+			{
+				return FLinearColor::Red;
+			}
+		}
+
+		return FSlateColor::UseForeground();
+	}
+
+	// Callback for getting the status icon image.
+	const FSlateBrush* HandleIconImage( ) const
+	{
+		ILauncherTaskPtr TaskPtr = Task.Pin();
+
+		if (TaskPtr.IsValid())
+		{
+			if (TaskPtr->GetStatus() == ELauncherTaskStatus::Canceled)
+			{
+				return FEditorStyle::GetBrush("Icons.Cross");
+			}
+
+			if (TaskPtr->GetStatus() == ELauncherTaskStatus::Completed)
+			{
+				return FEditorStyle::GetBrush("Symbols.Check");
+			}
+
+			if (TaskPtr->GetStatus() == ELauncherTaskStatus::Failed)
+			{
+				return FEditorStyle::GetBrush("Icons.Cross");
+			}
+		}
+
+		return NULL;
+	}
+
+	// Callback for getting the task's status text.
+	FString HandleStatusText( ) const
+	{
+		ILauncherTaskPtr TaskPtr = Task.Pin();
+
+		if (TaskPtr.IsValid())
+		{
+			ELauncherTaskStatus::Type TaskStatus = TaskPtr->GetStatus();
+
+			switch (TaskStatus)
+			{
+			case ELauncherTaskStatus::Busy:
+
+				return LOCTEXT("StatusInProgressText", "Busy").ToString();
+
+			case ELauncherTaskStatus::Canceled:
+
+				return LOCTEXT("StatusCanceledText", "Canceled").ToString();
+
+			case ELauncherTaskStatus::Canceling:
+
+				return LOCTEXT("StatusCancelingText", "Canceling").ToString();
+
+			case ELauncherTaskStatus::Completed:
+
+				return LOCTEXT("StatusCompletedText", "Completed").ToString();
+
+			case ELauncherTaskStatus::Failed:
+
+				return LOCTEXT("StatusFailedText", "Failed").ToString();
+
+			case ELauncherTaskStatus::Pending:
+
+				return LOCTEXT("StatusPendingText", "Pending").ToString();
+			}
+		}
+
+		return FString();
+	}
+
+	// Callback for determining the throbber's visibility.
+	EVisibility HandleThrobberVisibility( ) const
+	{
+		ILauncherTaskPtr TaskPtr = Task.Pin();
+
+		if (TaskPtr.IsValid())
+		{
+			if ((TaskPtr->GetStatus() == ELauncherTaskStatus::Busy) ||
+				(TaskPtr->GetStatus() == ELauncherTaskStatus::Canceling))
+			{
+				return EVisibility::Visible;
+			}
+		}
+
+		return EVisibility::Hidden;
+	}
+
+private:
+
+	// Holds a pointer to the task that is displayed in this row.
+	TWeakPtr<ILauncherTask> Task;
+};
+
+
+#undef LOCTEXT_NAMESPACE
diff --git a/Engine/Source/Developer/ProjectLauncher/Private/Widgets/SProjectLauncherValidation.h b/Engine/Source/Developer/ProjectLauncher/Private/Widgets/SProjectLauncherValidation.h
new file mode 100644
--- /dev/null
+++ b/Engine/Source/Developer/ProjectLauncher/Private/Widgets/SProjectLauncherValidation.h
@@ -0,0 +1,209 @@
+// Copyright 1998-2014 Epic Games, Inc. All Rights Reserved.
+
+#pragma once
+
+
+#define LOCTEXT_NAMESPACE "SProjectLauncherValidation"
+
+
+/**
+ * Implements the launcher's profile validation panel.
+ */
+class SProjectLauncherValidation
+	: public SCompoundWidget
+{
+public:
+
+	SLATE_BEGIN_ARGS(SProjectLauncherValidation) { }
+	SLATE_ATTRIBUTE(ILauncherProfilePtr, LaunchProfile)
+	SLATE_END_ARGS()
+
+
+public:
+
+	/**
+	 * Constructs the widget.
+	 *
+	 * @param InArgs The Slate argument list.
+	 * @param InModel The data model.
+	 */
+	void Construct( const FArguments& InArgs)
+	{
+		LaunchProfileAttr = InArgs._LaunchProfile;
+
+		ChildSlot
+		[
+			SNew(SVerticalBox)
+			
+			// build settings
+			+ SVerticalBox::Slot()
+				.AutoHeight()
+				[
+					MakeValidationMessage(TEXT("Icons.Error"), LOCTEXT("NoBuildGameSelectedError", "A Project must be selected.").ToString(), ELauncherProfileValidationErrors::NoProjectSelected)
+				]
+
+			+ SVerticalBox::Slot()
+				.AutoHeight()
+				[
+					MakeValidationMessage(TEXT("Icons.Error"), LOCTEXT("NoBuildConfigurationSelectedError", "A Build Configuration must be selected.").ToString(), ELauncherProfileValidationErrors::NoBuildConfigurationSelected)
+				]
+
+			// cook settings
+			+ SVerticalBox::Slot()
+				.AutoHeight()
+				[
+					MakeValidationMessage(TEXT("Icons.Error"), LOCTEXT("NoCookedPlatformSelectedError", "At least one Platform must be selected when cooking by the book.").ToString(), ELauncherProfileValidationErrors::NoPlatformSelected)
+				]
+
+			+ SVerticalBox::Slot()
+				.AutoHeight()
+				[
+					MakeValidationMessage(TEXT("Icons.Error"), LOCTEXT("NoCookedCulturesSelectedError", "At least one Culture must be selected when cooking by the book.").ToString(), ELauncherProfileValidationErrors::NoCookedCulturesSelected)
+				]
+
+			// packaging settings
+
+			// deployment settings
+			+ SVerticalBox::Slot()
+				.AutoHeight()
+				[
+					MakeValidationMessage(TEXT("Icons.Error"), LOCTEXT("CopyToDeviceRequiresCookByTheBookError", "Deployment by copying to device requires 'By The Book' cooking.").ToString(), ELauncherProfileValidationErrors::CopyToDeviceRequiresCookByTheBook)
+				]
+
+			+ SVerticalBox::Slot()
+				.AutoHeight()
+				[
+					MakeValidationMessage(TEXT("Icons.Error"), LOCTEXT("DeployedDeviceGroupRequired", "A device group must be selected when deploying builds.").ToString(), ELauncherProfileValidationErrors::DeployedDeviceGroupRequired)
+				]
+
+			// launch settings
+			+ SVerticalBox::Slot()
+				.AutoHeight()
+				[
+					MakeValidationMessage(TEXT("Icons.Error"), LOCTEXT("CustomRolesNotSupportedYet", "Custom launch roles are not supported yet.").ToString(), ELauncherProfileValidationErrors::CustomRolesNotSupportedYet)
+				]
+
+			+ SVerticalBox::Slot()
+				.AutoHeight()
+				[
+					MakeValidationMessage(TEXT("Icons.Error"), LOCTEXT("InitialCultureNotAvailable", "The Initial Culture selected for launch is not in the build.").ToString(), ELauncherProfileValidationErrors::InitialCultureNotAvailable)
+				]
+
+			+ SVerticalBox::Slot()
+				.AutoHeight()
+				[
+					MakeValidationMessage(TEXT("Icons.Error"), LOCTEXT("InitialMapNotAvailable", "The Initial Map selected for launch is not in the build.").ToString(), ELauncherProfileValidationErrors::InitialMapNotAvailable)
+				]
+
+			+ SVerticalBox::Slot()
+				.AutoHeight()
+				[
+					MakeValidationMessage(TEXT("Icons.Error"), LOCTEXT("NoLaunchRoleDeviceAssigned", "One or more launch roles do not have a device assigned.").ToString(), ELauncherProfileValidationErrors::NoLaunchRoleDeviceAssigned)
+				]
+		 
+			+ SVerticalBox::Slot()
+				.AutoHeight()
+				[
+					MakeCallbackMessage(TEXT("Icons.Error"), ELauncherProfileValidationErrors::NoPlatformSDKInstalled)
+				]
+		];
+	}
+
+protected:
+
+	/**
+	 * Creates a widget for a validation message.
+	 *
+	 * @param IconName The name of the message icon.
+	 * @param MessageText The message text.
+	 * @param MessageType The message type.
+	 */
+	TSharedRef<SWidget> MakeValidationMessage( const TCHAR* IconName, const FString& MessageText, ELauncherProfileValidationErrors::Type Message )
+	{
+		return SNew(SHorizontalBox)
+			.Visibility(this, &SProjectLauncherValidation::HandleValidationMessageVisibility, Message)
+
+		+ SHorizontalBox::Slot()
+			.AutoWidth()
+			.Padding(2.0)
+			[
+				SNew(SImage)
+					.Image(FEditorStyle::GetBrush(IconName))
+			]
+
+		+ SHorizontalBox::Slot()
+			.AutoWidth()
+			.VAlign(VAlign_Center)
+			[
+				SNew(STextBlock)
+					.Text(MessageText)
+			];
+	}
+
+	/**
+	 * Creates a widget for a validation message.
+	 *
+	 * @param IconName The name of the message icon.
+	 * @param MessageType The message type.
+	 */
+	TSharedRef<SWidget> MakeCallbackMessage( const TCHAR* IconName,ELauncherProfileValidationErrors::Type Message )
+	{
+		return SNew(SHorizontalBox)
+			.Visibility(this, &SProjectLauncherValidation::HandleValidationMessageVisibility, Message)
+
+		+ SHorizontalBox::Slot()
+			.AutoWidth()
+			.Padding(2.0)
+			[
+				SNew(SImage)
+					.Image(FEditorStyle::GetBrush(IconName))
+			]
+
+		+ SHorizontalBox::Slot()
+			.AutoWidth()
+			.VAlign(VAlign_Center)
+			[
+				SNew(STextBlock)
+					.Text(this, &SProjectLauncherValidation::HandleValidationMessage, Message)
+			];
+	}
+
+private:
+
+	// Callback for getting the visibility state of a validation message.
+	EVisibility HandleValidationMessageVisibility( ELauncherProfileValidationErrors::Type Error ) const
+	{
+		ILauncherProfilePtr LaunchProfile = LaunchProfileAttr.Get();
+		if (LaunchProfile.IsValid())
+		{
+			if (LaunchProfile->HasValidationError(Error))
+			{
+				return EVisibility::Visible;
+			}
+		}
+
+		return EVisibility::Collapsed;
+	}
+
+	FString	HandleValidationMessage( ELauncherProfileValidationErrors::Type Error ) const
+	{
+		ILauncherProfilePtr LaunchProfile = LaunchProfileAttr.Get();
+		if (LaunchProfile.IsValid())
+		{
+			if (LaunchProfile->HasValidationError(Error))
+			{
+				return LOCTEXT("NoPlatformSDKInstalled", "A required platform SDK is mising: ").ToString() + LaunchProfile->GetInvalidPlatform();
+			}
+		}
+
+		return TEXT("");
+	}
+
+private:
+
+	// Attribute for the launch profile this widget shows validation for. 
+	TAttribute<ILauncherProfilePtr> LaunchProfileAttr;
+};
+
+
+#undef LOCTEXT_NAMESPACE
diff --git a/Engine/Source/Developer/ProjectLauncher/Private/Widgets/Shared/SProjectLauncherBuildConfigurationSelector.h b/Engine/Source/Developer/ProjectLauncher/Private/Widgets/Shared/SProjectLauncherBuildConfigurationSelector.h
new file mode 100644
--- /dev/null
+++ b/Engine/Source/Developer/ProjectLauncher/Private/Widgets/Shared/SProjectLauncherBuildConfigurationSelector.h
@@ -0,0 +1,95 @@
+// Copyright 1998-2014 Epic Games, Inc. All Rights Reserved.
+
+#pragma once
+
+
+#define LOCTEXT_NAMESPACE "SProjectLauncherBuildConfigurationSelector"
+
+
+/**
+ * Delegate type for build configuration selections.
+ *
+ * The first parameter is the selected build configuration.
+ */
+DECLARE_DELEGATE_OneParam(FOnSessionSProjectLauncherBuildConfigurationSelected, EBuildConfigurations::Type)
+
+
+/**
+ * Implements a build configuration selector widget.
+ */
+class SProjectLauncherBuildConfigurationSelector
+	: public SCompoundWidget
+{
+public:
+
+	SLATE_BEGIN_ARGS(SProjectLauncherBuildConfigurationSelector) { }
+		SLATE_EVENT(FOnSessionSProjectLauncherBuildConfigurationSelected, OnConfigurationSelected)
+		SLATE_ATTRIBUTE(FString, Text)
+	SLATE_END_ARGS()
+
+public:
+
+	/**
+	 * Constructs the widget.
+	 *
+	 * @param InArgs The Slate argument list.
+	 * @param InModel The data model.
+	 */
+	void Construct(	const FArguments& InArgs )
+	{
+		OnConfigurationSelected = InArgs._OnConfigurationSelected;
+
+		// create build configurations menu
+		FMenuBuilder MenuBuilder(true, NULL);
+		{
+			FUIAction DebugAction(FExecuteAction::CreateSP(this, &SProjectLauncherBuildConfigurationSelector::HandleMenuEntryClicked, EBuildConfigurations::Debug));
+			MenuBuilder.AddMenuEntry(LOCTEXT("DebugAction", "Debug"), LOCTEXT("DebugActionHint", "Debug configuration."), FSlateIcon(), DebugAction);
+
+			FUIAction DebugGameAction(FExecuteAction::CreateSP(this, &SProjectLauncherBuildConfigurationSelector::HandleMenuEntryClicked, EBuildConfigurations::DebugGame));
+			MenuBuilder.AddMenuEntry(LOCTEXT("DebugGameAction", "DebugGame"), LOCTEXT("DebugGameActionHint", "DebugGame configuration."), FSlateIcon(), DebugGameAction);
+
+			FUIAction DevelopmentAction(FExecuteAction::CreateSP(this, &SProjectLauncherBuildConfigurationSelector::HandleMenuEntryClicked, EBuildConfigurations::Development));
+			MenuBuilder.AddMenuEntry(LOCTEXT("DevelopmentAction", "Development"), LOCTEXT("DevelopmentActionHint", "Development configuration."), FSlateIcon(), DevelopmentAction);
+
+			FUIAction ShippingAction(FExecuteAction::CreateSP(this, &SProjectLauncherBuildConfigurationSelector::HandleMenuEntryClicked, EBuildConfigurations::Shipping));
+			MenuBuilder.AddMenuEntry(LOCTEXT("ShippingAction", "Shipping"), LOCTEXT("ShippingActionHint", "Shipping configuration."), FSlateIcon(), ShippingAction);
+
+			FUIAction TestAction(FExecuteAction::CreateSP(this, &SProjectLauncherBuildConfigurationSelector::HandleMenuEntryClicked, EBuildConfigurations::Test));
+			MenuBuilder.AddMenuEntry(LOCTEXT("TestAction", "Test"), LOCTEXT("TestActionHint", "Test configuration."), FSlateIcon(), TestAction);
+		}
+
+		ChildSlot
+		[
+			// build configuration menu
+			SNew(SComboButton)
+			.VAlign(VAlign_Center)
+			.ButtonContent()
+			[
+				SNew(STextBlock)
+				.Font(FCoreStyle::Get().GetFontStyle(TEXT("SmallFont")))
+				.Text(InArgs._Text)
+			]
+			.ContentPadding(FMargin(4.0f, 2.0f))
+			.MenuContent()
+			[
+				MenuBuilder.MakeWidget()
+			]
+		];
+	}
+
+private:
+
+	// Callback for clicking a menu entry.
+	void HandleMenuEntryClicked( EBuildConfigurations::Type Configuration )
+	{
+		OnConfigurationSelected.ExecuteIfBound(Configuration);
+	}
+
+private:
+
+	// Holds a delegate to be invoked when a build configuration has been selected.
+	FOnSessionSProjectLauncherBuildConfigurationSelected OnConfigurationSelected;
+};
+
+
+#undef LOCTEXT_NAMESPACE
diff --git a/Engine/Source/Developer/ProjectLauncher/Private/Widgets/Shared/SProjectLauncherProfileLaunchButton.h b/Engine/Source/Developer/ProjectLauncher/Private/Widgets/Shared/SProjectLauncherProfileLaunchButton.h
new file mode 100644
--- /dev/null
+++ b/Engine/Source/Developer/ProjectLauncher/Private/Widgets/Shared/SProjectLauncherProfileLaunchButton.h
@@ -0,0 +1,160 @@
+// Copyright 1998-2014 Epic Games, Inc. All Rights Reserved.
+
+#pragma once
+
+#define LOCTEXT_NAMESPACE "SProjectLauncherProfileLaunchButton"
+
+/**
+ * Implements a build configuration selector widget.
+ */
+class SProjectLauncherProfileLaunchButton
+	: public SCompoundWidget
+{
+public:
+
+	SLATE_BEGIN_ARGS(SProjectLauncherProfileLaunchButton) { }
+	SLATE_EVENT(FOnClicked, OnClicked)
+		SLATE_ATTRIBUTE(ILauncherProfilePtr, LaunchProfile)
+	SLATE_END_ARGS()
+
+public:
+
+	/**
+	 * Constructs the widget.
+	 *
+	 * @param InArgs The Slate argument list.
+	 * @param InModel The data model.
+	 */
+	void Construct(const FArguments& InArgs, bool InShowText)
+	{
+		LaunchProfileAttr = InArgs._LaunchProfile;
+
+		TSharedPtr<SVerticalBox> VerticalBoxWidget;
+		ChildSlot
+		[	
+			SNew(SOverlay)
+					
+			+ SOverlay::Slot()
+			.HAlign(HAlign_Center)
+			.VAlign(VAlign_Center)
+			[
+				SNew(SButton)
+				.ButtonStyle(FCoreStyle::Get(), "Toolbar.Button")
+				.ForegroundColor(FEditorStyle::GetSlateColor("DefaultForeground"))
+				.OnClicked(InArgs._OnClicked)
+				.HAlign(HAlign_Center)
+				.VAlign(VAlign_Center)
+				.ContentPadding(0)
+				.IsEnabled(this, &SProjectLauncherProfileLaunchButton::ButtonEnabled)
+				[
+					SAssignNew(VerticalBoxWidget, SVerticalBox)
+
+					// Icon
+					+ SVerticalBox::Slot()
+					.AutoHeight()
+					.HAlign(HAlign_Center)
+					.VAlign(VAlign_Center)
+					[
+						SNew(SImage)
+						.Image(this, &SProjectLauncherProfileLaunchButton::GetLaunchIcon)
+					]
+				]
+			]
+
+			+ SOverlay::Slot()
+			.HAlign(HAlign_Center)
+			.VAlign(VAlign_Center)
+			[
+				SNew(SImage)
+				.Image(this, &SProjectLauncherProfileLaunchButton::GetErrorIcon)
+				.Visibility(this, &SProjectLauncherProfileLaunchButton::GetErrorVisibility)
+			]
+		];
+
+		// Add launch text is this was requested
+		if (InShowText && VerticalBoxWidget.IsValid())
+		{
+			VerticalBoxWidget->AddSlot()
+			.AutoHeight()
+			[
+				SNew(STextBlock)
+				.TextStyle(FCoreStyle::Get(), "Toolbar.Label")
+				.ShadowOffset(FVector2D::UnitVector)
+				.Text(LOCTEXT("ProjectLauncherLaunch", "Launch"))
+			];
+		}
+
+		// Tooltip when we have validation errors
+		SAssignNew(ErrorToolTipWidget, SToolTip)
+		[
+			SNew(SProjectLauncherValidation)
+			.LaunchProfile(InArgs._LaunchProfile)
+		];
+
+		// Otherwise we fall back on simple text
+		SetToolTipText(LOCTEXT("ProjectLauncherLaunchToolTip", "Launch this profile"));
+	}
+
+	virtual TSharedPtr<IToolTip> GetToolTip() override
+	{
+		ILauncherProfilePtr LaunchProfile = LaunchProfileAttr.Get();
+		if (LaunchProfile.IsValid() && LaunchProfile->HasValidationError())
+		{
+			ErrorToolTipWidget->SetContentWidget(
+				SNew(SProjectLauncherValidation)
+				.LaunchProfile(LaunchProfileAttr));
+
+			return ErrorToolTipWidget;
+		}
+		return SWidget::GetToolTip();
+	}
+
+private:
+
+	// @returns true if there is an error in the launch profile, no launch profile is considered an error
+	bool HasError() const
+	{
+		ILauncherProfilePtr LaunchProfile = LaunchProfileAttr.Get();
+		if (LaunchProfile.IsValid() && !LaunchProfile->HasValidationError())
+		{
+			return false;
+		}
+		return true;
+	}
+
+	// Callback to see if the launch button should be enabled
+	bool ButtonEnabled() const
+	{
+		return !HasError();
+	}
+
+	// Get the SlateIcon for Launch Button
+	const FSlateBrush* GetLaunchIcon() const
+	{
+		return FEditorStyle::GetBrush("Launcher.Run");
+	}
+
+	// Get the SlateIcon for Error Image
+	const FSlateBrush* GetErrorIcon() const
+	{
+		return FEditorStyle::GetBrush(TEXT("Icons.Error"));
+	}
+
+	// Callback to see if the error icon should be displayed, based on validity of the launch profile
+	EVisibility GetErrorVisibility() const
+	{
+		return HasError() ? EVisibility::Visible : EVisibility::Hidden;
+	}
+
+private:
+
+	// Attribute for the launch profile this widget launches. 
+	TAttribute<ILauncherProfilePtr> LaunchProfileAttr;
+
+	// Holds a pointer to our custom tooltip.
+	TSharedPtr<SToolTip> ErrorToolTipWidget;
+
+};
+
+
+#undef LOCTEXT_NAMESPACE
diff --git a/Engine/Source/Developer/ProjectLauncher/Private/Widgets/Shared/SProjectLauncherProfileNameDescEditor.h b/Engine/Source/Developer/ProjectLauncher/Private/Widgets/Shared/SProjectLauncherProfileNameDescEditor.h
new file mode 100644
--- /dev/null
+++ b/Engine/Source/Developer/ProjectLauncher/Private/Widgets/Shared/SProjectLauncherProfileNameDescEditor.h
@@ -0,0 +1,171 @@
+// Copyright 1998-2014 Epic Games, Inc. All Rights Reserved.
+
+#pragma once
+
+
+#define LOCTEXT_NAMESPACE "SProjectLauncherProfileNameDescEditor"
+
+/**
+ * Implements a build configuration selector widget.
+ */
+class SProjectLauncherProfileNameDescEditor
+	: public SCompoundWidget
+{
+public:
+
+	SLATE_BEGIN_ARGS(SProjectLauncherProfileNameDescEditor) { }
+		SLATE_ATTRIBUTE(ILauncherProfilePtr, LaunchProfile)
+	SLATE_END_ARGS()
+
+public:
+
+	/**
+	 * Constructs the widget.
+	 *
+	 * @param InArgs The Slate argument list.
+	 * @param InModel The data model.
+	 */
+	void Construct(const FArguments& InArgs, bool InShowAddDescriptionText)
+	{
+		EnterTextDescription = FText(LOCTEXT("LaunchProfileEnterDescription", "Enter a description here."));
+
+		LaunchProfileAttr = InArgs._LaunchProfile;
+		bShowAddDescriptionText = InShowAddDescriptionText;
+
+		ChildSlot
+		[
+			SNew(SHorizontalBox)
+
+			+ SHorizontalBox::Slot()
+			.AutoWidth()
+			.VAlign(VAlign_Center)
+			[
+				SNew(SBox)
+				.WidthOverride(40)
+				.HeightOverride(40)
+				[
+					SNew(SImage)
+					.Image(this, &SProjectLauncherProfileNameDescEditor::HandleProfileImage)
+				]
+			]
+
+			+ SHorizontalBox::Slot()
+			.FillWidth(1)
+			.VAlign(VAlign_Center)
+			[
+				SNew(SVerticalBox)
+
+				+ SVerticalBox::Slot()
+				.AutoHeight()
+				.Padding(2, 4, 2, 4)
+				[
+					SAssignNew(NameEditableTextBlock, SInlineEditableTextBlock)
+					.Text(this, &SProjectLauncherProfileNameDescEditor::OnGetNameText)
+					.OnTextCommitted(this, &SProjectLauncherProfileNameDescEditor::OnNameTextCommitted)
+					.Cursor(EMouseCursor::TextEditBeam)
+				]
+
+				+ SVerticalBox::Slot()
+				.AutoHeight()
+				.Padding(2, 4, 2, 4)
+				[
+					SNew(SInlineEditableTextBlock)
+					.Text(this, &SProjectLauncherProfileNameDescEditor::OnGetDescriptionText)
+					.Style(FCoreStyle::Get(), "InlineEditableTextBlockSmallStyle")
+					.OnTextCommitted(this, &SProjectLauncherProfileNameDescEditor::OnDescriptionTextCommitted)
+					.Cursor(EMouseCursor::TextEditBeam)
+				]
+			]
+		];
+	}
+
+	/**
+	 * Triggers a name edit for the profile.
+	 */
+	void TriggerNameEdit()
+	{
+		if (NameEditableTextBlock.IsValid())
+		{
+			NameEditableTextBlock->EnterEditingMode();
+		}
+	}
+
+private:
+
+	// Callback for getting the icon image of the profile.
+	const FSlateBrush* HandleProfileImage() const
+	{
+		//const PlatformInfo::FPlatformInfo* const PlatformInfo = PlatformInfo::FindPlatformInfo(*DeviceProxy->GetTargetPlatformName(SimpleProfile->GetDeviceVariant()));
+		//return (PlatformInfo) ? FEditorStyle::GetBrush(PlatformInfo->GetIconStyleName(PlatformInfo::EPlatformIconSize::Large)) : FStyleDefaults::GetNoBrush();
+		return FEditorStyle::GetBrush("LauncherCommand.QuickLaunch");
+	}
+
+	// Callback to get the name of the launch profile.
+	FText OnGetNameText() const
+	{
+		ILauncherProfilePtr LaunchProfile = LaunchProfileAttr.Get();
+		if (LaunchProfile.IsValid())
+		{
+			return FText::FromString(LaunchProfile->GetName());
+		}
+		return FText();
+	}
+
+	// Callback to set the name of the launch profile.
+	void OnNameTextCommitted(const FText& NewText, ETextCommit::Type InTextCommit)
+	{
+		ILauncherProfilePtr LaunchProfile = LaunchProfileAttr.Get();
+		if (LaunchProfile.IsValid())
+		{
+			LaunchProfile->SetName(NewText.ToString());
+		}
+	}
+
+	// Callback to get the description of the launch profile.
+	FText OnGetDescriptionText() const
+	{
+		ILauncherProfilePtr LaunchProfile = LaunchProfileAttr.Get();
+		if (LaunchProfile.IsValid())
+		{
+			FString Desc = LaunchProfile->GetDescription();
+			if (!bShowAddDescriptionText || !Desc.IsEmpty())
+			{
+				return FText::FromString(Desc);
+			}
+		}
+		return EnterTextDescription;
+	}
+
+	// Callback to set the description of the launch profile.
+	void OnDescriptionTextCommitted(const FText& NewText, ETextCommit::Type InTextCommit)
+	{
+		ILauncherProfilePtr LaunchProfile = LaunchProfileAttr.Get();
+		if (LaunchProfile.IsValid())
+		{
+			if (NewText.EqualTo(EnterTextDescription))
+			{
+				LaunchProfile->SetDescription("");
+			}
+			else
+			{
+				LaunchProfile->SetDescription(NewText.ToString());
+			}
+		}
+	}
+
+private:
+
+	// Attribute for the launch profile this widget edits. 
+	TAttribute<ILauncherProfilePtr> LaunchProfileAttr;
+
+	// Cache the no description enter suggestion text
+	FText EnterTextDescription;
+
+	// Whether we show an add description blurb when the profile has no description. 
+	bool bShowAddDescriptionText;
+
+	TSharedPtr<SInlineEditableTextBlock> NameEditableTextBlock;
+};
+
+
+#undef LOCTEXT_NAMESPACE
diff --git a/Engine/Source/Developer/ProjectLauncher/ProjectLauncher.Build.cs b/Engine/Source/Developer/ProjectLauncher/ProjectLauncher.Build.cs
new file mode 100644
--- /dev/null
+++ b/Engine/Source/Developer/ProjectLauncher/ProjectLauncher.Build.cs
@@ -0,0 +1,57 @@
+// Copyright 1998-2014 Epic Games, Inc. All Rights Reserved.
+
+using UnrealBuildTool;
+
+public class ProjectLauncher : ModuleRules
+{
+	public ProjectLauncher(TargetInfo Target)
+	{
+		PublicDependencyModuleNames.AddRange(
+			new string[] {
+				"Core",
+				"CoreUObject",
+				"TargetPlatform",
+			}
+		);
+
+		PublicIncludePathModuleNames.AddRange(
+			new string[] {
+				"LauncherServices",
+			}
+		);
+
+		PrivateDependencyModuleNames.AddRange(
+			new string[] {
+				"DesktopPlatform",
+                "InputCore",
+				"Slate",
+				"SlateCore",
+                "EditorStyle",
+			}
+		);
+
+		PrivateIncludePathModuleNames.AddRange(
+			new string[] {
+				"SessionServices",
+				"TargetDeviceServices",
+			}
+		);
+
+		PrivateIncludePaths.AddRange(
+			new string[] {
+				"Developer/ProjectLauncher/Private",
+				"Developer/ProjectLauncher/Private/Models",
+				"Developer/ProjectLauncher/Private/Widgets",
+				"Developer/ProjectLauncher/Private/Widgets/Build",
+				"Developer/ProjectLauncher/Private/Widgets/Cook",
+				"Developer/ProjectLauncher/Private/Widgets/Deploy",
+				"Developer/ProjectLauncher/Private/Widgets/Launch",
+				"Developer/ProjectLauncher/Private/Widgets/Package",
+                "Developer/ProjectLauncher/Private/Widgets/Profile",
+				"Developer/ProjectLauncher/Private/Widgets/Project",
+				"Developer/ProjectLauncher/Private/Widgets/Shared",
+				"Developer/ProjectLauncher/Private/Widgets/Toolbar",
+			}
+		);
+	}
+}
diff --git a/Engine/Source/Developer/ProjectLauncher/Public/Interfaces/IProjectLauncherModule.h b/Engine/Source/Developer/ProjectLauncher/Public/Interfaces/IProjectLauncherModule.h
new file mode 100644
--- /dev/null
+++ b/Engine/Source/Developer/ProjectLauncher/Public/Interfaces/IProjectLauncherModule.h
@@ -0,0 +1,28 @@
+// Copyright 1998-2014 Epic Games, Inc. All Rights Reserved.
+
+#pragma once
+
+
+/**
+ * Interface for launcher UI modules.
+ */
+class IProjectLauncherModule
+	: public IModuleInterface
+{
+public:
+	
+	/**
+	 * Creates a session launcher progress panel widget.
+	 *
+	 * @param LauncherWorker The launcher worker.
+	 * @return The new widget.
+	 */
+	virtual TSharedRef<class SWidget> CreateSProjectLauncherProgressPanel( const ILauncherWorkerRef& LauncherWorker ) = 0;
+
+public:
+
+	/**
+	 * Virtual destructor.
+	 */
+	virtual ~IProjectLauncherModule( ) { }
+};
diff --git a/Engine/Source/Developer/ProjectLauncher/Public/ProjectLauncher.h b/Engine/Source/Developer/ProjectLauncher/Public/ProjectLauncher.h
new file mode 100644
--- /dev/null
+++ b/Engine/Source/Developer/ProjectLauncher/Public/ProjectLauncher.h
@@ -0,0 +1,19 @@
+// Copyright 1998-2014 Epic Games, Inc. All Rights Reserved.
+
+#pragma once
+
+
+/* Dependencies
+ *****************************************************************************/
+
+#include "Core.h"
+#include "ModuleInterface.h"
+#include "LauncherServices.h"
+#include "Slate.h"
+#include "EditorStyle.h"
+
+
+/* Interfaces
+ *****************************************************************************/
+
+#include "IProjectLauncherModule.h"
diff --git a/Engine/Source/Developer/SessionLauncher/Private/Widgets/Deploy/SSessionLauncherDeployTargets.cpp b/Engine/Source/Developer/SessionLauncher/Private/Widgets/Deploy/SSessionLauncherDeployTargets.cpp
--- a/Engine/Source/Developer/SessionLauncher/Private/Widgets/Deploy/SSessionLauncherDeployTargets.cpp
+++ b/Engine/Source/Developer/SessionLauncher/Private/Widgets/Deploy/SSessionLauncherDeployTargets.cpp
@@ -288,7 +288,7 @@ void SSessionLauncherDeployTargets::RefreshDeviceProxyList( )
 
 		if (SelectedGroup.IsValid())
 		{
-			const TArray<FString>& Devices = SelectedGroup->GetDevices();
+			const TArray<FString>& Devices = SelectedGroup->GetDeviceIDs();
 
 			for (int32 Index = 0; Index < Devices.Num(); ++Index)
 			{
@@ -297,7 +297,7 @@ void SSessionLauncherDeployTargets::RefreshDeviceProxyList( )
 
 			if (ShowDevicesChoice == EShowDevicesChoices::ShowAllDevices)
 			{
-				Model->GetDeviceProxyManager()->GetProxies(FString(), false, DeviceProxyList);
+				Model->GetDeviceProxyManager()->GetProxies(NAME_None, false, DeviceProxyList);
 			}		
 		}
 	}
@@ -341,7 +341,7 @@ bool SSessionLauncherDeployTargets::HandleDeviceListRowIsEnabled( ITargetDeviceP
 
 		if (SelectedProfile.IsValid())
 		{
-			bool RetVal = SelectedProfile->IsDeployablePlatform(DeviceProxy->GetPlatformName()) || IsAlwaysEnabled;
+			bool RetVal = SelectedProfile->IsDeployablePlatform(DeviceProxy->GetTargetPlatformName(NAME_None)) || IsAlwaysEnabled;
 			return RetVal;
 		}
 	}
@@ -354,8 +354,8 @@ FText SSessionLauncherDeployTargets::HandleDeviceListRowToolTipText( ITargetDevi
 {
 	FTextBuilder Builder;
 	Builder.AppendLineFormat(LOCTEXT("DeviceListRowToolTipName", "Name: {0}"), FText::FromString(DeviceProxy->GetName()));
-	Builder.AppendLineFormat(LOCTEXT("DeviceListRowToolTipPlatform", "Platform: {0}"), FText::FromString(DeviceProxy->GetPlatformName()));
-	Builder.AppendLineFormat(LOCTEXT("DeviceListRowToolTipDeviceId", "Device ID: {0}"), FText::FromString(DeviceProxy->GetDeviceId()));
+	Builder.AppendLineFormat(LOCTEXT("DeviceListRowToolTipPlatform", "Platform: {0}"), FText::FromString(DeviceProxy->GetTargetPlatformName(NAME_None /*@Todo: Fix this!*/)));
+	Builder.AppendLineFormat(LOCTEXT("DeviceListRowToolTipDeviceId", "Device ID: {0}"), FText::FromString(DeviceProxy->GetTargetDeviceId(NAME_None /*@Todo: Fix this!*/)));
 
 	return Builder.ToText();
 }
diff --git a/Engine/Source/Developer/SessionLauncher/Private/Widgets/Deploy/SSessionLauncherDeviceListRow.h b/Engine/Source/Developer/SessionLauncher/Private/Widgets/Deploy/SSessionLauncherDeviceListRow.h
--- a/Engine/Source/Developer/SessionLauncher/Private/Widgets/Deploy/SSessionLauncherDeviceListRow.h
+++ b/Engine/Source/Developer/SessionLauncher/Private/Widgets/Deploy/SSessionLauncherDeviceListRow.h
@@ -138,11 +138,11 @@ private:
 		{
 			if (NewState == ESlateCheckBoxState::Checked)
 			{
-				ActiveGroup->AddDevice(DeviceProxy->GetDeviceId());
+				ActiveGroup->AddDevice(DeviceProxy->GetName());
 			}
 			else
 			{
-				ActiveGroup->RemoveDevice(DeviceProxy->GetDeviceId());
+				ActiveGroup->RemoveDevice(DeviceProxy->GetName());
 			}
 		}
 	}
@@ -156,7 +156,7 @@ private:
 
 			if (ActiveGroup.IsValid())
 			{
-				if (ActiveGroup->GetDevices().Contains(DeviceProxy->GetDeviceId()))
+				if (ActiveGroup->GetDeviceIDs().Contains(DeviceProxy->GetName()))
 				{
 					return ESlateCheckBoxState::Checked;
 				}
@@ -169,7 +169,7 @@ private:
 	// Callback for getting the icon image of the device.
 	const FSlateBrush* HandleDeviceImage( ) const
 	{
-		const PlatformInfo::FPlatformInfo* const PlatformInfo = PlatformInfo::FindPlatformInfo(*DeviceProxy->GetPlatformName());
+		const PlatformInfo::FPlatformInfo* const PlatformInfo = PlatformInfo::FindPlatformInfo(*DeviceProxy->GetTargetPlatformName(NAME_None));
 		return (PlatformInfo) ? FEditorStyle::GetBrush(PlatformInfo->GetIconStyleName(PlatformInfo::EPlatformIconSize::Normal)) : FStyleDefaults::GetNoBrush();
 	}
 
@@ -201,7 +201,7 @@ private:
 	// Callback for getting the host platform name.
 	FString HandleHostPlatformText( ) const
 	{
-		return DeviceProxy->GetPlatformName();
+		return DeviceProxy->GetTargetPlatformName(NAME_None);
 	}
 
 private:
diff --git a/Engine/Source/Developer/SessionLauncher/Private/Widgets/SSessionLauncherPreviewPage.cpp b/Engine/Source/Developer/SessionLauncher/Private/Widgets/SSessionLauncherPreviewPage.cpp
--- a/Engine/Source/Developer/SessionLauncher/Private/Widgets/SSessionLauncherPreviewPage.cpp
+++ b/Engine/Source/Developer/SessionLauncher/Private/Widgets/SSessionLauncherPreviewPage.cpp
@@ -815,7 +815,7 @@ void SSessionLauncherPreviewPage::RefreshDeviceProxyList( )
 
 		if (DeployedDeviceGroup.IsValid())
 		{
-			const TArray<FString>& Devices = DeployedDeviceGroup->GetDevices();
+			const TArray<FString>& Devices = DeployedDeviceGroup->GetDeviceIDs();
 
 			for (int32 Index = 0; Index < Devices.Num(); ++Index)
 			{
diff --git a/Engine/Source/Developer/TargetDeviceServices/Classes/TargetDeviceServiceMessages.h b/Engine/Source/Developer/TargetDeviceServices/Classes/TargetDeviceServiceMessages.h
--- a/Engine/Source/Developer/TargetDeviceServices/Classes/TargetDeviceServiceMessages.h
+++ b/Engine/Source/Developer/TargetDeviceServices/Classes/TargetDeviceServiceMessages.h
@@ -19,6 +19,12 @@ struct FTargetDeviceServiceDeployCommit
 {
 	GENERATED_USTRUCT_BODY()
 
+	/**
+	 * Holds the variant identifier of the target device to use.
+	 */
+	UPROPERTY()
+	FName Variant;
+
 	/**
 	 * Holds the identifier of the deployment transaction to commit.
 	 */
@@ -29,13 +35,14 @@ struct FTargetDeviceServiceDeployCommit
 	/**
 	 * Default constructor.
 	 */
-	FTargetDeviceServiceDeployCommit( ) { }
+	FTargetDeviceServiceDeployCommit() { }
 
 	/**
 	 * Creates and initializes a new instance.
 	 */
-	FTargetDeviceServiceDeployCommit( const FGuid& InTransactionId )
-		: TransactionId(InTransactionId)
+	FTargetDeviceServiceDeployCommit(FName InVariant, const FGuid& InTransactionId)
+		: Variant(InVariant)
+		, TransactionId(InTransactionId)
 	{ }
 };
 
@@ -61,13 +68,13 @@ USTRUCT()
 struct FTargetDeviceServiceDeployFile
 {
 	GENERATED_USTRUCT_BODY()
-	
+
 	/**
 	 * Holds the name and path of the file as it will be stored on the target device.
 	 */
 	UPROPERTY()
 	FString TargetFileName;
-	
+
 	/**
 	 * Holds the identifier of the deployment transaction that this file is part of.
 	 */
@@ -78,12 +85,12 @@ struct FTargetDeviceServiceDeployFile
 	/**
 	 * Default constructor.
 	 */
-	FTargetDeviceServiceDeployFile( ) { }
+	FTargetDeviceServiceDeployFile() { }
 
 	/**
 	 * Creates and initializes a new instance.
 	 */
-	FTargetDeviceServiceDeployFile( const FString& InTargetFileName, const FGuid& InTransactionId )
+	FTargetDeviceServiceDeployFile(const FString& InTargetFileName, const FGuid& InTransactionId)
 		: TargetFileName(InTargetFileName)
 		, TransactionId(InTransactionId)
 	{ }
@@ -110,6 +117,12 @@ struct FTargetDeviceServiceDeployFinished
 {
 	GENERATED_USTRUCT_BODY()
 
+	/**
+	 * Holds the variant identifier of the target device to use.
+	 */
+	UPROPERTY()
+	FName Variant;
+
 	/**
 	 * Holds the created identifier for the deployed application.
 	 *
@@ -137,13 +150,14 @@ struct FTargetDeviceServiceDeployFinished
 	/**
 	 * Default constructor.
 	 */
-	FTargetDeviceServiceDeployFinished( ) { }
+	FTargetDeviceServiceDeployFinished() { }
 
 	/**
 	 * Creates and initializes a new instance.
 	 */
-	FTargetDeviceServiceDeployFinished( const FString& InAppId, bool InSucceeded, FGuid InTransactionId )
-		: AppID(InAppId)
+	FTargetDeviceServiceDeployFinished(FName InVariant, const FString& InAppId, bool InSucceeded, FGuid InTransactionId)
+		: Variant(InVariant)
+		, AppID(InAppId)
 		, Succeeded(InSucceeded)
 		, TransactionId(InTransactionId)
 	{ }
@@ -160,7 +174,7 @@ struct TStructOpsTypeTraits<FTargetDeviceServiceDeployFinished> : public TStruct
 
 
 /* Application launch messages
- *****************************************************************************/
+*****************************************************************************/
 
 /**
  * Implements a message for committing a deployment transaction.
@@ -174,6 +188,12 @@ struct FTargetDeviceServiceLaunchApp
 {
 	GENERATED_USTRUCT_BODY()
 
+	/**
+	 * Holds the variant identifier of the target device to use.
+	 */
+	UPROPERTY()
+	FName Variant;
+
 	/**
 	 * Holds the identifier of the application to launch.
 	 *
@@ -201,13 +221,14 @@ struct FTargetDeviceServiceLaunchApp
 	/**
 	 * Default constructor.
 	 */
-	FTargetDeviceServiceLaunchApp( ) { }
+	FTargetDeviceServiceLaunchApp() { }
 
 	/**
 	 * Creates and initializes a new instance.
 	 */
-	FTargetDeviceServiceLaunchApp( const FString& InAppId, uint8 InBuildConfiguration, const FString& InParams )
-		: AppID(InAppId)
+	FTargetDeviceServiceLaunchApp(FName InVariant, const FString& InAppId, uint8 InBuildConfiguration, const FString& InParams)
+		: Variant(InVariant)
+		, AppID(InAppId)
 		, BuildConfiguration(InBuildConfiguration)
 		, Params(InParams)
 	{ }
@@ -260,12 +281,12 @@ struct FTargetDeviceServiceLaunchFinished
 	/**
 	 * Default constructor.
 	 */
-	FTargetDeviceServiceLaunchFinished( ) { }
+	FTargetDeviceServiceLaunchFinished() { }
 
 	/**
 	 * Creates and initializes a new instance.
 	 */
-	FTargetDeviceServiceLaunchFinished( const FString& InAppId, int32 InProcessId, bool InSucceeded )
+	FTargetDeviceServiceLaunchFinished(const FString& InAppId, int32 InProcessId, bool InSucceeded)
 		: AppID(InAppId)
 		, ProcessId(InProcessId)
 		, Succeeded(InSucceeded)
@@ -281,7 +302,6 @@ struct TStructOpsTypeTraits<FTargetDeviceServiceLaunchFinished> : public TStruct
 	};
 };
 
-
 /* Device claiming messages
  *****************************************************************************/
 
@@ -300,7 +320,7 @@ struct FTargetDeviceClaimDenied
 	 * Holds the identifier of the device that is already claimed.
 	 */
 	UPROPERTY()
-	FString DeviceID;
+	FString DeviceName;
 
 	/**
 	 * Holds the name of the host computer that claimed the device.
@@ -322,8 +342,8 @@ struct FTargetDeviceClaimDenied
 	/**
 	 * Creates and initializes a new instance.
 	 */
-	FTargetDeviceClaimDenied( const FString& InDeviceId, const FString& InHostName, const FString& InHostUser )
-		: DeviceID(InDeviceId)
+	FTargetDeviceClaimDenied(const FString& InDeviceName, const FString& InHostName, const FString& InHostUser)
+		: DeviceName(InDeviceName)
 		, HostName(InHostName)
 		, HostUser(InHostUser)
 	{ }
@@ -354,7 +374,7 @@ struct FTargetDeviceClaimed
 	 * Holds the identifier of the device that is being claimed.
 	 */
 	UPROPERTY()
-	FString DeviceID;
+	FString DeviceName;
 
 	/**
 	 * Holds the name of the host computer that is claiming the device.
@@ -376,8 +396,8 @@ struct FTargetDeviceClaimed
 	/**
 	 * Creates and initializes a new instance.
 	 */
-	FTargetDeviceClaimed( const FString& InDeviceId, const FString& InHostName, const FString& InHostUser )
-		: DeviceID(InDeviceId)
+	FTargetDeviceClaimed(const FString& InDeviceName, const FString& InHostName, const FString& InHostUser)
+		: DeviceName(InDeviceName)
 		, HostName(InHostName)
 		, HostUser(InHostUser)
 	{ }
@@ -408,7 +428,7 @@ struct FTargetDeviceUnclaimed
 	 * Holds the identifier of the device that is no longer claimed.
 	 */
 	UPROPERTY()
-	FString DeviceID;
+	FString DeviceName;
 
 	/**
 	 * Holds the name of the host computer that had claimed the device.
@@ -430,8 +450,8 @@ struct FTargetDeviceUnclaimed
 	/**
 	 * Creates and initializes a new instance.
 	 */
-	FTargetDeviceUnclaimed( const FString& InDeviceId, const FString& InHostName, const FString& InHostUser )
-		: DeviceID(InDeviceId)
+	FTargetDeviceUnclaimed(const FString& InDeviceName, const FString& InHostName, const FString& InHostUser)
+		: DeviceName(InDeviceName)
 		, HostName(InHostName)
 		, HostUser(InHostUser)
 	{ }
@@ -487,6 +507,34 @@ struct TStructOpsTypeTraits<FTargetDeviceServicePing> : public TStructOpsTypeTra
 };
 
 
+/**
+* Struct for a flavor's information
+*/
+USTRUCT()
+struct FTargetDeviceVariant
+{
+	GENERATED_USTRUCT_BODY()
+
+	UPROPERTY()
+	FString DeviceID;
+
+	UPROPERTY()
+	FName VariantName;
+
+	UPROPERTY()
+	FString TargetPlatformName;
+
+	UPROPERTY()
+	FName TargetPlatformId;
+
+	UPROPERTY()
+	FName VanillaPlatformId;
+
+	UPROPERTY()
+	FText PlatformDisplayName;
+};
+
+
 /**
  * Implements a message that is sent in response to target device service discovery messages.
  */
@@ -501,12 +549,6 @@ struct FTargetDeviceServicePong
 	UPROPERTY()
 	bool Connected;
 
-	/**
-	 * Holds the identifier of the device that this message corresponds to.
-	 */
-	UPROPERTY()
-	FString DeviceID;
-
 	/**
 	 * Holds the name of the host computer that the device is attached to.
 	 */
@@ -549,12 +591,6 @@ struct FTargetDeviceServicePong
 	UPROPERTY()
 	FString DeviceUserPassword;
 
-	/**
-	 * Holds the name of the device platform, i.e. Windows, Mac or iOS.
-	 */
-	UPROPERTY()
-	FString PlatformName;
-
 	/**
 	 * Holds a flag indicating whether this device is shared with other users on the network.
 	 */
@@ -585,11 +621,29 @@ struct FTargetDeviceServicePong
 	UPROPERTY()
 	bool SupportsReboot;
 
+	/**
+	* Holds a flag indicating whether the device's target platform supports variants.
+	*/
+	UPROPERTY()
+	bool SupportsVariants;
+
 	/**
 	 * Holds the device type.
 	 */
 	UPROPERTY()
 	FString Type;
+
+	/**
+	 * Holds the variant name of the default device.
+	 */
+	UPROPERTY()
+	FName DefaultVariant;
+
+	/**
+	 * List of the Flavors this service supports
+	 */
+	UPROPERTY()
+	TArray<FTargetDeviceVariant> Variants;
 };
 
 template<>
@@ -740,6 +794,12 @@ struct FTargetDeviceServiceRunExecutable
 {
 	GENERATED_USTRUCT_BODY()
 
+	/**
+	 * Holds the variant identifier of the target device to use for execution.
+	 */
+	UPROPERTY()
+	FName Variant;
+
 	/**
 	 * Holds the path to the executable on the device.
 	 */
@@ -756,13 +816,14 @@ struct FTargetDeviceServiceRunExecutable
 	/**
 	 * Default constructor.
 	 */
-	FTargetDeviceServiceRunExecutable( ) { }
+	FTargetDeviceServiceRunExecutable() { }
 
 	/**
 	 * Creates and initializes a new instance.
 	 */
-	FTargetDeviceServiceRunExecutable( const FString& InExecutablePath, const FString& InParams )
-		: ExecutablePath(InExecutablePath)
+	FTargetDeviceServiceRunExecutable(FName InVariant, const FString& InExecutablePath, const FString& InParams)
+		: Variant(InVariant)
+		, ExecutablePath(InExecutablePath)
 		, Params(InParams)
 	{ }
 };
@@ -787,6 +848,12 @@ struct FTargetDeviceServiceRunFinished
 {
 	GENERATED_USTRUCT_BODY()
 
+	/**
+	 * Holds the variant identifier of the target device to use.
+	 */
+	UPROPERTY()
+	FName Variant;
+
 	/**
 	 * Holds the path to the executable that was run.
 	 */
@@ -809,13 +876,14 @@ struct FTargetDeviceServiceRunFinished
 	/**
 	 * Default constructor.
 	 */
-	FTargetDeviceServiceRunFinished( ) { }
+	FTargetDeviceServiceRunFinished() { }
 
 	/**
 	 * Creates and initializes a new instance.
 	 */
-	FTargetDeviceServiceRunFinished( const FString& InExecutablePath, int32 InProcessId, bool InSucceeded )
-		: ExecutablePath(InExecutablePath)
+	FTargetDeviceServiceRunFinished(FName InVariant, const FString& InExecutablePath, int32 InProcessId, bool InSucceeded)
+		: Variant(InVariant)
+		, ExecutablePath(InExecutablePath)
 		, ProcessId(InProcessId)
 		, Succeeded(InSucceeded)
 	{ }
diff --git a/Engine/Source/Developer/TargetDeviceServices/Private/TargetDeviceProxy.cpp b/Engine/Source/Developer/TargetDeviceServices/Private/TargetDeviceProxy.cpp
--- a/Engine/Source/Developer/TargetDeviceServices/Private/TargetDeviceProxy.cpp
+++ b/Engine/Source/Developer/TargetDeviceServices/Private/TargetDeviceProxy.cpp
@@ -6,30 +6,23 @@
 /* FTargetDeviceProxy structors
  *****************************************************************************/
 
-FTargetDeviceProxy::FTargetDeviceProxy( const FString& InId )
+FTargetDeviceProxy::FTargetDeviceProxy(const FString& InName)
 	: Connected(false)
-	, Id(InId)
 	, LastUpdateTime(0)
+	, Name(InName)
 	, SupportsMultiLaunch(false)
 	, SupportsPowerOff(false)
 	, SupportsPowerOn(false)
 	, SupportsReboot(false)
 {
-	FTargetDeviceId DeviceId;
-
-	if (FTargetDeviceId::Parse(InId, DeviceId))
-	{
-		Name = DeviceId.GetDeviceName();
-		Platform = DeviceId.GetPlatformName();
-	}
 
 	InitializeMessaging();
 }
 
 
-FTargetDeviceProxy::FTargetDeviceProxy( const FString& InId, const FTargetDeviceServicePong& Message, const IMessageContextRef& Context )
+FTargetDeviceProxy::FTargetDeviceProxy(const FString& InName, const FTargetDeviceServicePong& Message, const IMessageContextRef& Context)
 	: Connected(false)
-	, Id(InId)
+	, Name(InName)
 	, SupportsMultiLaunch(false)
 	, SupportsPowerOff(false)
 	, SupportsPowerOn(false)
@@ -45,25 +38,39 @@ FTargetDeviceProxy::FTargetDeviceProxy( const FString& InId, const FTargetDevice
 
 void FTargetDeviceProxy::UpdateFromMessage( const FTargetDeviceServicePong& Message, const IMessageContextRef& Context )
 {
-	if (Message.DeviceID == Id)
+	if (Message.Name == Name)
 	{
 		MessageAddress = Context->GetSender();
 
 		Connected = Message.Connected;
 		HostName = Message.HostName;
 		HostUser = Message.HostUser;
-		Id = Message.DeviceID;
 		Make = Message.Make;
 		Model = Message.Model;
 		Name = Message.Name;
 		DeviceUser = Message.DeviceUser;
 		DeviceUserPassword = Message.DeviceUserPassword;
-		Platform = Message.PlatformName;
 		Shared = Message.Shared;
 		SupportsMultiLaunch = Message.SupportsMultiLaunch;
 		SupportsPowerOff = Message.SupportsPowerOff;
 		SupportsPowerOn = Message.SupportsPowerOn;
 		SupportsReboot = Message.SupportsReboot;
+		SupportsVariants = Message.SupportsVariants;
+		DefaultVariant = Message.DefaultVariant;
+
+		// Update the map of flavors.
+		for (int Index = 0; Index < Message.Variants.Num(); Index++)
+		{
+			const FTargetDeviceVariant & MsgVariant = Message.Variants[Index];
+
+			FTargetDeviceProxyVariant & Variant = TargetDeviceVariants.Add(MsgVariant.VariantName);
+			Variant.DeviceID = MsgVariant.DeviceID;
+			Variant.VariantName = MsgVariant.VariantName;
+			Variant.TargetPlatformName = MsgVariant.TargetPlatformName;
+			Variant.TargetPlatformId = MsgVariant.TargetPlatformId;
+			Variant.VanillaPlatformId = MsgVariant.VanillaPlatformId;
+			Variant.PlatformDisplayName = MsgVariant.PlatformDisplayName;
+		}
 
 		LastUpdateTime = FDateTime::UtcNow();
 	}
@@ -73,7 +80,101 @@ void FTargetDeviceProxy::UpdateFromMessage( const FTargetDeviceServicePong& Mess
 /* ITargetDeviceProxyPtr interface
  *****************************************************************************/
 
-bool FTargetDeviceProxy::DeployApp( const TMap<FString, FString>& Files, const FGuid& TransactionId )
+int32 FTargetDeviceProxy::GetNumVariants() const
+{
+	return TargetDeviceVariants.Num();
+}
+
+int32 FTargetDeviceProxy::GetVariants(TArray<FName>& OutVariants) const
+{
+	return TargetDeviceVariants.GetKeys(OutVariants);
+}
+
+FName FTargetDeviceProxy::GetTargetDeviceVariant(const FString& InDeviceId) const
+{
+	for (TMap<FName, FTargetDeviceProxyVariant>::TConstIterator ItVariant(TargetDeviceVariants); ItVariant; ++ItVariant)
+	{
+		const FTargetDeviceProxyVariant & Variant = ItVariant.Value();
+		if (Variant.DeviceID == InDeviceId)
+		{
+			return ItVariant.Key();
+		}
+	}
+	return NAME_None;
+}
+
+bool FTargetDeviceProxy::HasDeviceId(const FString& InDeviceId) const
+{
+	for (TMap<FName, FTargetDeviceProxyVariant>::TConstIterator ItVariant(TargetDeviceVariants); ItVariant; ++ItVariant)
+	{
+		const FTargetDeviceProxyVariant & Variant = ItVariant.Value();
+		if (Variant.DeviceID == InDeviceId)
+		{
+			return true;
+		}
+	}
+	return false;
+}
+
+const FString& FTargetDeviceProxy::GetTargetDeviceId(FName InVariant) const
+{
+	if (InVariant == NAME_None)
+	{
+		return TargetDeviceVariants[DefaultVariant].DeviceID;
+	}
+	return TargetDeviceVariants[InVariant].DeviceID;
+}
+
+bool FTargetDeviceProxy::HasTargetPlatform(FName InTargetPlatformId) const
+{
+	for (TMap<FName, FTargetDeviceProxyVariant>::TConstIterator ItVariant(TargetDeviceVariants); ItVariant; ++ItVariant)
+	{
+		const FTargetDeviceProxyVariant & Variant = ItVariant.Value();
+		if (Variant.TargetPlatformId == InTargetPlatformId)
+		{
+			return true;
+		}
+	}
+	return false;
+}
+
+FString FTargetDeviceProxy::GetTargetPlatformName(FName InVariant) const
+{
+	if (InVariant == NAME_None)
+	{
+		return TargetDeviceVariants[DefaultVariant].TargetPlatformName;
+	}
+	return TargetDeviceVariants[InVariant].TargetPlatformName;
+}
+
+FName FTargetDeviceProxy::GetTargetPlatformId(FName InVariant) const
+{
+	if (InVariant == NAME_None)
+	{
+		return TargetDeviceVariants[DefaultVariant].TargetPlatformId;
+	}
+	return TargetDeviceVariants[InVariant].TargetPlatformId;
+}
+
+FName FTargetDeviceProxy::GetVanillaPlatformId(FName InVariant) const
+{
+	if (InVariant == NAME_None)
+	{
+		return TargetDeviceVariants[DefaultVariant].VanillaPlatformId;
+	}
+	return TargetDeviceVariants[InVariant].VanillaPlatformId;
+}
+
+FText FTargetDeviceProxy::GetPlatformDisplayName(FName InVariant) const
+{
+	if (InVariant == NAME_None)
+	{
+		return TargetDeviceVariants[DefaultVariant].PlatformDisplayName;
+	}
+	return TargetDeviceVariants[InVariant].PlatformDisplayName;
+}
+
+bool FTargetDeviceProxy::DeployApp(FName InVariant, const TMap<FString, FString>& Files, const FGuid& TransactionId)
 {
 	for (TMap<FString, FString>::TConstIterator It(Files); It; ++It)
 	{
@@ -83,26 +184,23 @@ bool FTargetDeviceProxy::DeployApp( const TMap<FString, FString>& Files, const F
 		MessageEndpoint->Send(new FTargetDeviceServiceDeployFile(It.Value(), TransactionId), FileAttachment, MessageAddress);
 	}
 
-	MessageEndpoint->Send(new FTargetDeviceServiceDeployCommit(TransactionId), MessageAddress);
+	MessageEndpoint->Send(new FTargetDeviceServiceDeployCommit(InVariant, TransactionId), MessageAddress);
 
 	return true;
 }
 
-
-bool FTargetDeviceProxy::LaunchApp( const FString& AppId, EBuildConfigurations::Type BuildConfiguration, const FString& Params )
+bool FTargetDeviceProxy::LaunchApp(FName InVariant, const FString& AppId, EBuildConfigurations::Type BuildConfiguration, const FString& Params)
 {
-	MessageEndpoint->Send(new FTargetDeviceServiceLaunchApp(AppId, BuildConfiguration, Params), MessageAddress);
+	MessageEndpoint->Send(new FTargetDeviceServiceLaunchApp(InVariant, AppId, BuildConfiguration, Params), MessageAddress);
 
 	return true;
 }
 
-
 void FTargetDeviceProxy::PowerOff( bool Force )
 {
 	MessageEndpoint->Send(new FTargetDeviceServicePowerOff(FPlatformProcess::UserName(false), Force), MessageAddress);
 }
 
-
 void FTargetDeviceProxy::PowerOn( )
 {
 	MessageEndpoint->Send(new FTargetDeviceServicePowerOn(FPlatformProcess::UserName(false)), MessageAddress);
@@ -114,10 +212,9 @@ void FTargetDeviceProxy::Reboot( )
 	MessageEndpoint->Send(new FTargetDeviceServiceReboot(FPlatformProcess::UserName(false)), MessageAddress);
 }
 
-
-void FTargetDeviceProxy::Run( const FString& ExecutablePath, const FString& Params )
+void FTargetDeviceProxy::Run(FName InVariant, const FString& ExecutablePath, const FString& Params)
 {
-	MessageEndpoint->Send(new FTargetDeviceServiceRunExecutable(ExecutablePath, Params), MessageAddress);
+	MessageEndpoint->Send(new FTargetDeviceServiceRunExecutable(InVariant, ExecutablePath, Params), MessageAddress);
 }
 
 
@@ -126,16 +223,15 @@ void FTargetDeviceProxy::Run( const FString& ExecutablePath, const FString& Para
 
 void FTargetDeviceProxy::InitializeMessaging( )
 {
-	MessageEndpoint = FMessageEndpoint::Builder(FName(*FString::Printf(TEXT("FTargetDeviceProxy (%s)"), *Id)))
+	MessageEndpoint = FMessageEndpoint::Builder(FName(*FString::Printf(TEXT("FTargetDeviceProxy (%s)"), *Name)))
 		.Handling<FTargetDeviceServiceDeployFinished>(this, &FTargetDeviceProxy::HandleDeployFinishedMessage)
 		.Handling<FTargetDeviceServiceLaunchFinished>(this, &FTargetDeviceProxy::HandleLaunchFinishedMessage);
 }
 
-
 /* FTargetDeviceProxy event handlers
- *****************************************************************************/
+*****************************************************************************/
 
-void FTargetDeviceProxy::HandleDeployFinishedMessage( const FTargetDeviceServiceDeployFinished& Message, const IMessageContextRef& Context )
+void FTargetDeviceProxy::HandleDeployFinishedMessage(const FTargetDeviceServiceDeployFinished& Message, const IMessageContextRef& Context)
 {
 	if (Message.Succeeded)
 	{
@@ -148,7 +244,7 @@ void FTargetDeviceProxy::HandleDeployFinishedMessage( const FTargetDeviceService
 }
 
 
-void FTargetDeviceProxy::HandleLaunchFinishedMessage( const FTargetDeviceServiceLaunchFinished& Message, const IMessageContextRef& Context )
+void FTargetDeviceProxy::HandleLaunchFinishedMessage(const FTargetDeviceServiceLaunchFinished& Message, const IMessageContextRef& Context)
 {
 	if (Message.Succeeded)
 	{
@@ -158,4 +254,4 @@ void FTargetDeviceProxy::HandleLaunchFinishedMessage( const FTargetDeviceService
 	{
 		LaunchFailedDelegate.Broadcast(Message.AppID);
 	}
-}
+}
\ No newline at end of file
diff --git a/Engine/Source/Developer/TargetDeviceServices/Private/TargetDeviceProxy.h b/Engine/Source/Developer/TargetDeviceServices/Private/TargetDeviceProxy.h
--- a/Engine/Source/Developer/TargetDeviceServices/Private/TargetDeviceProxy.h
+++ b/Engine/Source/Developer/TargetDeviceServices/Private/TargetDeviceProxy.h
@@ -21,14 +21,14 @@ public:
 	/**
 	 * Default constructor
 	 */
-	FTargetDeviceProxy( ) { }
+	FTargetDeviceProxy() { }
 
 	/**
 	 * Creates and initializes a new instance.
 	 *
 	 * @param InId The identifier of the target device to create this proxy for.
 	 */
-	FTargetDeviceProxy( const FString& InId );
+	FTargetDeviceProxy(const FString& InId);
 
 	/**
 	 * Creates and initializes a new instance.
@@ -37,7 +37,7 @@ public:
 	 * @param Message The message to initialize from.
 	 * @param Context The message context.
 	 */
-	FTargetDeviceProxy( const FString& InId, const FTargetDeviceServicePong& Message, const IMessageContextRef& Context );
+	FTargetDeviceProxy(const FString& InId, const FTargetDeviceServicePong& Message, const IMessageContextRef& Context);
 
 public:
 
@@ -46,7 +46,7 @@ public:
 	 *
 	 * @return Date and time of the last update.
 	 */
-	const FDateTime& GetLastUpdateTime( ) const
+	const FDateTime& GetLastUpdateTime() const
 	{
 		return LastUpdateTime;
 	}
@@ -57,133 +57,149 @@ public:
 	 * @param Message The message containing the response.
 	 * @param Context The message context.
 	 */
-	void UpdateFromMessage( const FTargetDeviceServicePong& Message, const IMessageContextRef& Context );
+	void UpdateFromMessage(const FTargetDeviceServicePong& Message, const IMessageContextRef& Context);
 
 public:
 
 	// ITargetDeviceProxy interface
 
-	virtual const bool CanMultiLaunch( ) const override
+	virtual const bool CanMultiLaunch() const override
 	{
 		return SupportsMultiLaunch;
 	}
 
-	virtual bool CanPowerOff( ) const override
+	virtual bool CanPowerOff() const override
 	{
 		return SupportsPowerOff;
 	}
 
-	virtual bool CanPowerOn( ) const override
+	virtual bool CanPowerOn() const override
 	{
 		return SupportsPowerOn;
 	}
 
-	virtual bool CanReboot( ) const override
+	virtual bool CanReboot() const override
 	{
 		return SupportsReboot;
 	}
 
-	virtual bool DeployApp( const TMap<FString, FString>& Files, const FGuid& TransactionId ) override;
-
-	virtual const FString& GetDeviceId( ) override
+	virtual bool CanSupportVariants() const override
 	{
-		return Id;
+		return SupportsVariants;
 	}
 
-	virtual const FString& GetHostName( ) const override
+	virtual int32 GetNumVariants() const override;
+
+	virtual int32 GetVariants(TArray<FName>& OutVariants) const override;
+
+	virtual FName GetTargetDeviceVariant(const FString& InDeviceId) const override;
+
+	virtual bool HasDeviceId(const FString& InDeviceId) const override;
+
+	virtual const FString& GetTargetDeviceId(FName InVariant) const override;
+
+	virtual bool HasTargetPlatform(FName InTargetPlatformId) const override;
+
+	virtual FString GetTargetPlatformName(FName InVariant) const override;
+
+	virtual FName GetTargetPlatformId(FName InVariant) const override;
+
+	virtual FName GetVanillaPlatformId(FName InVariant) const override;
+
+	virtual FText GetPlatformDisplayName(FName InVariant) const override;
+
+	virtual const FString& GetHostName() const override
 	{
 		return HostName;
 	}
 
-	virtual const FString& GetHostUser( ) const override
+	virtual const FString& GetHostUser() const override
 	{
 		return HostUser;
 	}
 
-	virtual const FString& GetMake( ) const override
+	virtual const FString& GetMake() const override
 	{
 		return Make;
 	}
 
-	virtual const FString& GetModel( ) const override
+	virtual const FString& GetModel() const override
 	{
 		return Model;
 	}
 
-	virtual const FString& GetName( ) const override
+	virtual const FString& GetName() const override
 	{
 		return Name;
 	}
 
-	virtual const FString& GetDeviceUser( ) const override
+	virtual const FString& GetDeviceUser() const override
 	{
 		return DeviceUser;
 	}
 
-	virtual const FString& GetDeviceUserPassword( ) const override
+	virtual const FString& GetDeviceUserPassword() const override
 	{
 		return DeviceUserPassword;
 	}
 
-	virtual const FString& GetPlatformName( ) const override
-	{
-		return Platform;
-	}
-
-	virtual const FString& GetType( ) const override
+	virtual const FString& GetType() const override
 	{
 		return Type;
 	}
 
-	virtual bool IsConnected( ) const override
+	virtual bool IsConnected() const override
 	{
 		return Connected;
 	}
 
-	virtual bool IsShared( ) const override
+	virtual bool IsShared() const override
 	{
 		return Shared;
 	}
 
-	virtual bool LaunchApp( const FString& AppId, EBuildConfigurations::Type BuildConfiguration, const FString& Params ) override;
+	virtual bool DeployApp(FName InVariant, const TMap<FString, FString>& Files, const FGuid& TransactionId) override;
 
-	virtual FOnTargetDeviceProxyDeployCommitted& OnDeployCommitted( ) override
+	virtual bool LaunchApp(FName InVariant, const FString& AppId, EBuildConfigurations::Type BuildConfiguration, const FString& Params) override;
+
+	virtual FOnTargetDeviceProxyDeployCommitted& OnDeployCommitted() override
 	{
 		return DeployCommittedDelegate;
 	}
 
-	virtual FOnTargetDeviceProxyDeployFailed& OnDeployFailed( ) override
+	virtual FOnTargetDeviceProxyDeployFailed& OnDeployFailed() override
 	{
 		return DeployFailedDelegate;
 	}
 
-	virtual FOnTargetDeviceProxyLaunchFailed& OnLaunchFailed( ) override
+	virtual FOnTargetDeviceProxyLaunchFailed& OnLaunchFailed() override
 	{
 		return LaunchFailedDelegate;
 	}
 
-	virtual void PowerOff( bool Force ) override;
+	virtual void PowerOff(bool Force) override;
+
+	virtual void PowerOn() override;
 
-	virtual void PowerOn( ) override;
+	virtual void Reboot() override;
 
-	virtual void Reboot( ) override;
+	virtual void Run(FName InVariant, const FString& ExecutablePath, const FString& Params) override;
 
-	virtual void Run( const FString& ExecutablePath, const FString& Params ) override;
 
 protected:
 
 	/**
 	 * Initializes the message endpoint.
 	 */
-	void InitializeMessaging( );
+	void InitializeMessaging();
 
 private:
 
 	// Handles FTargetDeviceServiceDeployFinishedMessage messages.
-	void HandleDeployFinishedMessage( const FTargetDeviceServiceDeployFinished& Message, const IMessageContextRef& Context );
+	void HandleDeployFinishedMessage(const FTargetDeviceServiceDeployFinished& Message, const IMessageContextRef& Context);
 
 	// Handles FTargetDeviceServiceLaunchFinishedMessage messages.
-	void HandleLaunchFinishedMessage( const FTargetDeviceServiceLaunchFinished& Message, const IMessageContextRef& Context );
+	void HandleLaunchFinishedMessage(const FTargetDeviceServiceLaunchFinished& Message, const IMessageContextRef& Context);
 
 private:
 
@@ -196,9 +212,6 @@ private:
 	// Holds the name of the user that owns the device.
 	FString HostUser;
 
-	// Holds the remote device's identifier.
-	FString Id;
-
 	// Holds the time at which the last ping reply was received.
 	FDateTime LastUpdateTime;
 
@@ -217,9 +230,6 @@ private:
 	// Holds the name of the device.
 	FString Name;
 
-	// Holds the device's platform name.
-	FString Platform;
-
 	// Holds device user
 	FString DeviceUser;
 
@@ -241,9 +251,35 @@ private:
 	// Holds a flag indicating whether the device can reboot.
 	bool SupportsReboot;
 
+	// Holds a flag indicating whether the device's target platform supports variants.
+	bool SupportsVariants;
+
 	// Holds the device type.
 	FString Type;
 
+	// Holds default variant name.
+	FName DefaultVariant;
+
+	// Holds data about a device proxy variant
+	class FTargetDeviceProxyVariant
+	{
+	public:
+		// Holds a string version of the variants device id.
+		FString DeviceID;
+
+		// Holds the variant name, this is the the map key as well
+		FName VariantName;
+
+		// Platform information
+		FString TargetPlatformName;
+		FName TargetPlatformId;
+		FName VanillaPlatformId;
+		FText PlatformDisplayName;
+	};
+
+	// Map of all the Variants for this Device
+	TMap<FName, FTargetDeviceProxyVariant> TargetDeviceVariants;
+
 private:
 
 	// Holds a delegate to be invoked when a build has been deployed to the target device.
diff --git a/Engine/Source/Developer/TargetDeviceServices/Private/TargetDeviceProxyManager.cpp b/Engine/Source/Developer/TargetDeviceServices/Private/TargetDeviceProxyManager.cpp
--- a/Engine/Source/Developer/TargetDeviceServices/Private/TargetDeviceProxyManager.cpp
+++ b/Engine/Source/Developer/TargetDeviceServices/Private/TargetDeviceProxyManager.cpp
@@ -30,19 +30,19 @@ FTargetDeviceProxyManager::~FTargetDeviceProxyManager( )
 /* ITargetDeviceProxyLocator interface
  *****************************************************************************/
 
-ITargetDeviceProxyPtr FTargetDeviceProxyManager::FindProxy( const FString& DeviceId ) 
+ITargetDeviceProxyPtr FTargetDeviceProxyManager::FindProxy( const FString& Name ) 
 {
-	return Proxies.FindRef(DeviceId);
+	return Proxies.FindRef(Name);
 }
 
 
-ITargetDeviceProxyRef FTargetDeviceProxyManager::FindOrAddProxy( const FString& DeviceId )
+ITargetDeviceProxyRef FTargetDeviceProxyManager::FindOrAddProxy(const FString& Name)
 {
-	TSharedPtr<FTargetDeviceProxy>& Proxy = Proxies.FindOrAdd(DeviceId);
+	TSharedPtr<FTargetDeviceProxy>& Proxy = Proxies.FindOrAdd(Name);
 
 	if (!Proxy.IsValid())
 	{
-		Proxy = MakeShareable(new FTargetDeviceProxy(DeviceId));
+		Proxy = MakeShareable(new FTargetDeviceProxy(Name));
 
 		ProxyAddedDelegate.Broadcast(Proxy.ToSharedRef());
 	}
@@ -50,8 +50,22 @@ ITargetDeviceProxyRef FTargetDeviceProxyManager::FindOrAddProxy( const FString&
 	return Proxy.ToSharedRef();
 }
 
+ITargetDeviceProxyPtr FTargetDeviceProxyManager::FindProxyDeviceForTargetDevice(const FString& DeviceId)
+{
+	for (TMap<FString, TSharedPtr<FTargetDeviceProxy> >::TConstIterator ItProxies(Proxies); ItProxies; ++ItProxies)
+	{
+		const TSharedPtr<FTargetDeviceProxy>& Proxy = ItProxies.Value();
 
-void FTargetDeviceProxyManager::GetProxies( const FString& PlatformName, bool IncludeUnshared, TArray<ITargetDeviceProxyPtr>& OutProxies )
+		if (Proxy->HasDeviceId(DeviceId))
+		{
+			return Proxy;
+		}
+	}
+
+	return ITargetDeviceProxyPtr();
+}
+
+void FTargetDeviceProxyManager::GetProxies(FName TargetPlatformName, bool IncludeUnshared, TArray<ITargetDeviceProxyPtr>& OutProxies)
 {
 	OutProxies.Reset();
 
@@ -61,15 +75,14 @@ void FTargetDeviceProxyManager::GetProxies( const FString& PlatformName, bool In
 
 		if ((IncludeUnshared || Proxy->IsShared()) || (Proxy->GetHostUser() == FPlatformProcess::UserName(false)))
 		{
-			if (PlatformName.IsEmpty() || (Proxy->GetPlatformName() == PlatformName))
+			if (TargetPlatformName == NAME_None || Proxy->HasTargetPlatform(TargetPlatformName))
 			{
 				OutProxies.Add(Proxy);
-			}			
+			}		
 		}
 	}
 }
 
-
 /* FTargetDeviceProxyManager implementation
  *****************************************************************************/
 
@@ -89,7 +102,7 @@ void FTargetDeviceProxyManager::RemoveDeadProxies( )
 }
 
 
-void FTargetDeviceProxyManager::SendPing( )
+void FTargetDeviceProxyManager::SendPing()
 {
 	if (MessageEndpoint.IsValid())
 	{
@@ -103,11 +116,11 @@ void FTargetDeviceProxyManager::SendPing( )
 
 void FTargetDeviceProxyManager::HandlePongMessage( const FTargetDeviceServicePong& Message, const IMessageContextRef& Context )
 {
-	TSharedPtr<FTargetDeviceProxy>& Proxy = Proxies.FindOrAdd(Message.DeviceID);
+	TSharedPtr<FTargetDeviceProxy>& Proxy = Proxies.FindOrAdd(Message.Name);
 
 	if (!Proxy.IsValid())
 	{
-		Proxy = MakeShareable(new FTargetDeviceProxy(Message.DeviceID, Message, Context));
+		Proxy = MakeShareable(new FTargetDeviceProxy(Message.Name, Message, Context));
 		ProxyAddedDelegate.Broadcast(Proxy.ToSharedRef());
 	}
 	else
diff --git a/Engine/Source/Developer/TargetDeviceServices/Private/TargetDeviceProxyManager.h b/Engine/Source/Developer/TargetDeviceServices/Private/TargetDeviceProxyManager.h
--- a/Engine/Source/Developer/TargetDeviceServices/Private/TargetDeviceProxyManager.h
+++ b/Engine/Source/Developer/TargetDeviceServices/Private/TargetDeviceProxyManager.h
@@ -25,9 +25,12 @@ public:
 
 	// ITargetDeviceProxyLocator interface
 
-	virtual ITargetDeviceProxyRef FindOrAddProxy( const FString& DeviceId ) override;
-	virtual ITargetDeviceProxyPtr FindProxy( const FString& DeviceId ) override;
-	virtual void GetProxies( const FString& PlatformName, bool IncludeUnshared, TArray<ITargetDeviceProxyPtr>& OutProxies ) override;
+	virtual ITargetDeviceProxyRef FindOrAddProxy(const FString& Name) override;
+	virtual ITargetDeviceProxyPtr FindProxy(const FString& Name) override;
+	
+	virtual ITargetDeviceProxyPtr FindProxyDeviceForTargetDevice(const FString& DeviceId) override;
+	
+	virtual void GetProxies(FName TargetPlatformName, bool IncludeUnshared, TArray<ITargetDeviceProxyPtr>& OutProxies) override;
 
 	DECLARE_DERIVED_EVENT(FTargetDeviceProxyManager, ITargetDeviceProxyManager::FOnTargetDeviceProxyAdded, FOnTargetDeviceProxyAdded);
 	virtual FOnTargetDeviceProxyAdded& OnProxyAdded( ) override
diff --git a/Engine/Source/Developer/TargetDeviceServices/Private/TargetDeviceService.cpp b/Engine/Source/Developer/TargetDeviceServices/Private/TargetDeviceService.cpp
new file mode 100644
--- /dev/null
+++ b/Engine/Source/Developer/TargetDeviceServices/Private/TargetDeviceService.cpp
@@ -0,0 +1,110 @@
+// Copyright 1998-2014 Epic Games, Inc. All Rights Reserved.
+
+#pragma once
+
+#include "TargetDeviceServicesPrivatePCH.h"
+
+#include "PlatformInfo.h"
+
+void FTargetDeviceService::AddTargetDevice(ITargetDevicePtr InDevice)
+{
+	FName Variant = FName(InDevice->GetTargetPlatform().PlatformName().GetCharArray().GetData());
+
+	if (DevicePlatformName == NAME_None)
+	{
+		// If this seems nasty your right!
+		// This is just one more nastiness in this class due to the fact that we intend to refactor the target platform stuff as a separate task.
+		const PlatformInfo::FPlatformInfo& Info = InDevice->GetTargetPlatform().GetPlatformInfo();
+		DevicePlatformName = Info.PlatformInfoName;
+		const PlatformInfo::FPlatformInfo* VanillaInfo = PlatformInfo::FindVanillaPlatformInfo(Info.VanillaPlatformName);
+		DevicePlatformDisplayName = VanillaInfo->DisplayName.ToString();
+		
+		// Sigh the hacks... Should be able to remove if platform info gets cleaned up.... Windows doesn't have a reasonable vanilla platform.
+		const FString VariableSplit(TEXT("("));
+		FString Full = VanillaInfo->DisplayName.ToString();
+		FString Left;
+		FString Right;
+		bool bSplit = Full.Split(VariableSplit, &Left, &Right);
+		DevicePlatformDisplayName = bSplit ? Left.Trim() : Full;
+	}
+
+	// Make sure we don't have conflicting flavors for the same device!
+	// Could also be a double add, which ideally would be avoided.
+	check(!(TargetDevicePtrs.FindRef(Variant).IsValid()));
+
+	TargetDevicePtrs.Add(Variant, InDevice);
+
+	// Sort and choose cache the default
+	TargetDevicePtrs.ValueSort(FVariantSortCallback());
+	auto DeviceIterator = TargetDevicePtrs.CreateIterator();
+	DefaultDevicePtr = (*DeviceIterator).Value;
+}
+
+void FTargetDeviceService::RemoveTargetDevice(ITargetDevicePtr InDevice)
+{
+	FName Variant = FName(InDevice->GetTargetPlatform().PlatformName().GetCharArray().GetData());
+
+	TargetDevicePtrs.Remove(Variant);
+
+	// Cache the default
+	auto DeviceIterator = TargetDevicePtrs.CreateIterator();
+	DefaultDevicePtr = (*DeviceIterator).Value;
+}
+
+void FTargetDeviceService::HandlePingMessage(const FTargetDeviceServicePing& InMessage, const IMessageContextRef& Context)
+{
+	if (!Running)
+	{
+		return;
+	}
+
+	if (Shared || (InMessage.HostUser == FPlatformProcess::UserName(false)))
+	{
+		ITargetDevicePtr DefaultDevice = GetDevice(); // Default Device is needed here!
+
+		if (DefaultDevice.IsValid())
+		{
+			const FString& PlatformName = DefaultDevice->GetTargetPlatform().PlatformName();
+			const PlatformInfo::FPlatformInfo* VanillaInfo = PlatformInfo::FindVanillaPlatformInfo(FName(*PlatformName));
+
+			FTargetDeviceServicePong* Message = new FTargetDeviceServicePong();
+
+			Message->Name = DefaultDevice->GetName();
+			Message->Type = ETargetDeviceTypes::ToString(DefaultDevice->GetDeviceType());
+			Message->HostName = FPlatformProcess::ComputerName();
+			Message->HostUser = FPlatformProcess::UserName(false);
+			Message->Connected = DefaultDevice->IsConnected();
+			Message->Make = TEXT("@todo");
+			Message->Model = TEXT("@todo");
+			DefaultDevice->GetUserCredentials(Message->DeviceUser, Message->DeviceUserPassword);
+			Message->Shared = Shared;
+			Message->SupportsMultiLaunch = DefaultDevice->SupportsFeature(ETargetDeviceFeatures::MultiLaunch);
+			Message->SupportsPowerOff = DefaultDevice->SupportsFeature(ETargetDeviceFeatures::PowerOff);
+			Message->SupportsPowerOn = DefaultDevice->SupportsFeature(ETargetDeviceFeatures::PowerOn);
+			Message->SupportsReboot = DefaultDevice->SupportsFeature(ETargetDeviceFeatures::Reboot);
+			Message->SupportsVariants = DefaultDevice->GetTargetPlatform().SupportsVariants();
+			Message->DefaultVariant = FName(DefaultDevice->GetTargetPlatform().PlatformName().GetCharArray().GetData());
+
+			// Add the data for all the flavors
+			Message->Variants.SetNumZeroed(TargetDevicePtrs.Num());
+
+			int Index = 0;
+			for (auto TargetDeviceIt = TargetDevicePtrs.CreateIterator(); TargetDeviceIt; ++TargetDeviceIt, ++Index)
+			{
+				const ITargetDevicePtr& TargetDevice = TargetDeviceIt.Value().Pin();
+				const PlatformInfo::FPlatformInfo& Info = TargetDevice->GetTargetPlatform().GetPlatformInfo();
+
+				FTargetDeviceVariant& Variant = Message->Variants[Index];
+
+				Variant.DeviceID = TargetDevice->GetId().ToString();
+				Variant.VariantName = TargetDeviceIt.Key();
+				Variant.TargetPlatformName = TargetDevice->GetTargetPlatform().PlatformName();
+				Variant.TargetPlatformId = Info.TargetPlatformName;
+				Variant.VanillaPlatformId = Info.VanillaPlatformName;
+				Variant.PlatformDisplayName = Info.DisplayName;
+			}
+
+			MessageEndpoint->Send(Message, Context->GetSender());
+		}
+	}
+}
\ No newline at end of file
diff --git a/Engine/Source/Developer/TargetDeviceServices/Private/TargetDeviceService.h b/Engine/Source/Developer/TargetDeviceServices/Private/TargetDeviceService.h
--- a/Engine/Source/Developer/TargetDeviceServices/Private/TargetDeviceService.h
+++ b/Engine/Source/Developer/TargetDeviceServices/Private/TargetDeviceService.h
@@ -2,33 +2,39 @@
 
 #pragma once
 
-
 /**
  * Implements remote services for a specific target device.
  */
 class FTargetDeviceService
 	: public ITargetDeviceService
 {
+	struct FVariantSortCallback
+	{
+		FORCEINLINE bool operator()(const ITargetDeviceWeakPtr& A, const ITargetDeviceWeakPtr& B) const
+		{
+			ITargetDevicePtr APtr = A.Pin();
+			ITargetDevicePtr BPtr = B.Pin();
+			return APtr->GetTargetPlatform().GetVariantPriority() > BPtr->GetTargetPlatform().GetVariantPriority();
+		}
+	};
+
 public:
 
 	/**
 	 * Creates and initializes a new instance.
 	 *
-	 * @param InDeviceId The identifier of the target device to expose.
 	 * @param InDeviceName The name of the device to expose.
 	 * @param InMessageBus The message bus to listen on for clients.
 	 */
-	FTargetDeviceService( const FTargetDeviceId& InDeviceId, const FString& InDeviceName, const IMessageBusRef& InMessageBus )
-		: CachedDeviceName(InDeviceName)
-		, DeviceId(InDeviceId)
+	FTargetDeviceService(const FString& InDeviceName, const IMessageBusRef& InMessageBus)
+		: DeviceName(InDeviceName)
 		, Running(false)
 		, Shared(false)
 	{
 		// initialize messaging
-		MessageEndpoint = FMessageEndpoint::Builder(FName(*FString::Printf(TEXT("FTargetDeviceService (%s)"), *DeviceId.ToString())), InMessageBus)
+		MessageEndpoint = FMessageEndpoint::Builder(FName(*FString::Printf(TEXT("FTargetDeviceService (%s)"), *DeviceName)), InMessageBus)
 			.Handling<FTargetDeviceClaimDenied>(this, &FTargetDeviceService::HandleClaimDeniedMessage)
 			.Handling<FTargetDeviceClaimed>(this, &FTargetDeviceService::HandleClaimedMessage)
-			.Handling<FTargetDeviceUnclaimed>(this, &FTargetDeviceService::HandleUnclaimedMessage)
 			.Handling<FTargetDeviceServiceDeployCommit>(this, &FTargetDeviceService::HandleDeployCommitMessage)
 			.Handling<FTargetDeviceServiceDeployFile>(this, &FTargetDeviceService::HandleDeployFileMessage)
 			.Handling<FTargetDeviceServiceLaunchApp>(this, &FTargetDeviceService::HandleLaunchAppMessage)
@@ -59,9 +65,18 @@ public:
 
 	// ITargetDeviceService interface
 
-	virtual bool CanStart( ) const override
+	virtual void AddTargetDevice(ITargetDevicePtr InDevice) override;
+
+	virtual void RemoveTargetDevice(ITargetDevicePtr InDevice) override;
+
+	virtual int32 NumTargetDevices() override
+	{
+		return TargetDevicePtrs.Num();
+	}
+
+	virtual bool CanStart(FName InFlavor = NAME_None) const override
 	{
-		ITargetDevicePtr TargetDevice = TargetDevicePtr.Pin();
+		ITargetDevicePtr TargetDevice = GetDevice(InFlavor);
 
 		if (TargetDevice.IsValid())
 		{
@@ -71,46 +86,52 @@ public:
 		return false;
 	}
 
-	virtual const FString& GetClaimHost( ) override
+	virtual const FString& GetClaimHost() override
 	{
 		return ClaimHost;
 	}
 
-	virtual const FString& GetClaimUser( ) override
+	virtual const FString& GetClaimUser() override
 	{
 		return ClaimUser;
 	}
 
-	virtual ITargetDevicePtr GetDevice( ) override
+	virtual ITargetDevicePtr GetDevice(FName InVariant = NAME_None) const override
 	{
-		return AcquireDevice();
+		ITargetDevicePtr TargetDevice;
+		if (InVariant == NAME_None)
+		{
+			TargetDevice = DefaultDevicePtr.Pin();
+		}
+		else
+		{
+			const ITargetDeviceWeakPtr * WeakTargetDevicePtr = TargetDevicePtrs.Find(InVariant);
+			if (WeakTargetDevicePtr != NULL)
+			{
+				TargetDevice = WeakTargetDevicePtr->Pin();
+			}
+		}
+		return TargetDevice;
 	}
 
-	virtual const FTargetDeviceId& GetDeviceId( ) const override
+	virtual FString GetDeviceName() const override
 	{
-		return DeviceId;
+		return DeviceName;
 	}
 
-	virtual FString GetCachedDeviceName( ) const override
+	virtual FName GetDevicePlatformName() const override
 	{
-		ITargetDevicePtr TargetDevice = TargetDevicePtr.Pin();
-
-		if (TargetDevice.IsValid())
-		{
-			return TargetDevice->GetName();
-		}
-
-		if (CachedDeviceName.IsEmpty())
-		{
-			return DeviceId.ToString();
-		}
+		return DevicePlatformName;
+	}
 
-		return CachedDeviceName;
+	virtual FString GetDevicePlatformDisplayName() const override
+	{
+		return DevicePlatformDisplayName;
 	}
 
-	virtual bool IsRunning( ) const override
+	virtual bool IsRunning() const override
 	{
-		return (Running && TargetDevicePtr.IsValid());
+		return Running;
 	}
 
 	virtual bool IsShared( ) const override
@@ -118,31 +139,21 @@ public:
 		return (Running && Shared);
 	}
 
-	virtual void SetShared( bool InShared ) override
+	virtual void SetShared(bool InShared) override
 	{
 		Shared = InShared;
 	}
 
-	virtual bool Start( ) override
+	virtual bool Start() override
 	{
 		if (!Running && MessageEndpoint.IsValid())
 		{
-			// initializes target device
-			ITargetDevicePtr TargetDevice = AcquireDevice();
-
-			if (!TargetDevice.IsValid())
-			{
-				return false;
-			}
-
-			CachedDeviceName = TargetDevice->GetName();
-
 			// notify other services
 			ClaimAddress = MessageEndpoint->GetAddress();
 			ClaimHost = FPlatformProcess::ComputerName();
 			ClaimUser = FPlatformProcess::UserName(false);
 
-			MessageEndpoint->Publish(new FTargetDeviceClaimed(DeviceId.ToString(), ClaimHost, ClaimUser));
+			MessageEndpoint->Publish(new FTargetDeviceClaimed(DeviceName, ClaimHost, ClaimUser));
 
 			Running = true;
 		}
@@ -150,20 +161,12 @@ public:
 		return true;
 	}
 
-	virtual void Stop( ) override
+	virtual void Stop() override
 	{
 		if (Running)
 		{
-			// cache device details
-			ITargetDevicePtr TargetDevice = TargetDevicePtr.Pin();
-
-			if (TargetDevice.IsValid())
-			{
-				CachedDeviceName = TargetDevice->GetName();
-			}
-
 			// notify other services
-			MessageEndpoint->Publish(new FTargetDeviceUnclaimed(DeviceId.ToString(), FPlatformProcess::ComputerName(), FPlatformProcess::UserName(false)));
+			MessageEndpoint->Publish(new FTargetDeviceUnclaimed(DeviceName, FPlatformProcess::ComputerName(), FPlatformProcess::UserName(false)));
 
 			Running = false;
 		}
@@ -172,30 +175,12 @@ public:
 protected:
 
 	/**
-	 * Acquires the target device, if needed.
-	 *
-	 * @return The target device.
-	 */
-	ITargetDevicePtr AcquireDevice( )
-	{
-		ITargetDevicePtr TargetDevice = TargetDevicePtr.Pin();
-
-		if (!TargetDevice.IsValid())
-		{
-			TargetDevice = GetTargetPlatformManager()->FindTargetDevice(DeviceId);
-			TargetDevicePtr = TargetDevice;
-		}
-
-		return TargetDevice;
-	}
-
-	/**
-	 * Stores the specified file to deploy.
-	 *
-	 * @param FileReader The archive reader providing access to the file data.
-	 * @param TargetFileName The desired name of the file on the target device.
-	 */
-	bool StoreDeployedFile( FArchive* FileReader, const FString& TargetFileName  ) const
+	* Stores the specified file to deploy.
+	*
+	* @param FileReader The archive reader providing access to the file data.
+	* @param TargetFileName The desired name of the file on the target device.
+	*/
+	bool StoreDeployedFile(FArchive* FileReader, const FString& TargetFileName) const
 	{
 		if (FileReader == NULL)
 		{
@@ -222,7 +207,7 @@ protected:
 		}
 
 		void* Buffer = FMemory::Malloc(BufferSize);
-	
+
 		while (BytesRemaining > 0)
 		{
 			FileReader->Serialize(Buffer, BufferSize);
@@ -249,7 +234,7 @@ private:
 	// Callback for FTargetDeviceClaimDenied messages.
 	void HandleClaimDeniedMessage( const FTargetDeviceClaimDenied& Message, const IMessageContextRef& Context )
 	{
-		if (Running && (Message.DeviceID == DeviceId.ToString()))
+		if (Running && (Message.DeviceName == DeviceName))
 		{
 			Stop();
 
@@ -262,13 +247,13 @@ private:
 	// Callback for FTargetDeviceClaimDenied messages.
 	void HandleClaimedMessage( const FTargetDeviceClaimed& Message, const IMessageContextRef& Context )
 	{
-		if (Message.DeviceID == DeviceId.ToString())
+		if (Message.DeviceName == DeviceName)
 		{
 			if (Running)
 			{
 				if (Context->GetSender() != MessageEndpoint->GetAddress())
 				{
-					MessageEndpoint->Send(new FTargetDeviceClaimDenied(DeviceId.ToString(), FPlatformProcess::ComputerName(), FPlatformProcess::UserName(false)), Context->GetSender());
+					MessageEndpoint->Send(new FTargetDeviceClaimDenied(DeviceName, FPlatformProcess::ComputerName(), FPlatformProcess::UserName(false)), Context->GetSender());
 				}
 			}
 			else
@@ -283,7 +268,7 @@ private:
 	// Callback for FTargetDeviceClaimDropped messages.
 	void HandleUnclaimedMessage( const FTargetDeviceUnclaimed& Message, const IMessageContextRef& Context )
 	{
-		if (Message.DeviceID == DeviceId.ToString())
+		if (Message.DeviceName == DeviceName)
 		{
 			if (Context->GetSender() == ClaimAddress)
 			{
@@ -295,7 +280,7 @@ private:
 	}
 
 	// Callback for FTargetDeviceServiceDeployFile messages.
-	void HandleDeployFileMessage( const FTargetDeviceServiceDeployFile& Message, const IMessageContextRef& Context )
+	void HandleDeployFileMessage(const FTargetDeviceServiceDeployFile& Message, const IMessageContextRef& Context)
 	{
 		if (!Running)
 		{
@@ -321,15 +306,14 @@ private:
 	}
 
 	// Callback for FTargetDeviceServiceDeployCommit messages.
-	void HandleDeployCommitMessage( const FTargetDeviceServiceDeployCommit& Message, const IMessageContextRef& Context )
+	void HandleDeployCommitMessage(const FTargetDeviceServiceDeployCommit& Message, const IMessageContextRef& Context)
 	{
 		if (!Running)
 		{
 			return;
 		}
 
-		ITargetDevicePtr TargetDevice = TargetDevicePtr.Pin();
-
+		ITargetDevicePtr TargetDevice = GetDevice(Message.Variant);
 		if (TargetDevice.IsValid())
 		{
 			FString SourceFolder = FPaths::EngineIntermediateDir() / TEXT("Deploy") / Message.TransactionId.ToString();
@@ -338,20 +322,19 @@ private:
 			bool Succeeded = TargetDevice->Deploy(SourceFolder, OutAppId);
 
 			IFileManager::Get().DeleteDirectory(*SourceFolder, false, true);
-			MessageEndpoint->Send(new FTargetDeviceServiceDeployFinished(OutAppId, Succeeded, Message.TransactionId), Context->GetSender());
+			MessageEndpoint->Send(new FTargetDeviceServiceDeployFinished(Message.Variant, OutAppId, Succeeded, Message.TransactionId), Context->GetSender());
 		}
 	}
 
 	// Callback for FTargetDeviceServiceLaunchApp messages.
-	void HandleLaunchAppMessage( const FTargetDeviceServiceLaunchApp& Message, const IMessageContextRef& Context )
+	void HandleLaunchAppMessage(const FTargetDeviceServiceLaunchApp& Message, const IMessageContextRef& Context)
 	{
 		if (!Running)
 		{
 			return;
 		}
 
-		ITargetDevicePtr TargetDevice = AcquireDevice();
-
+		ITargetDevicePtr TargetDevice = GetDevice(Message.Variant);
 		if (TargetDevice.IsValid())
 		{
 			uint32 ProcessId;
@@ -365,41 +348,7 @@ private:
 	}
 
 	// Callback for FTargetDeviceServicePing messages.
-	void HandlePingMessage( const FTargetDeviceServicePing& InMessage, const IMessageContextRef& Context )
-	{
-		if (!Running)
-		{
-			return;
-		}
-
-		if (Shared || (InMessage.HostUser == FPlatformProcess::UserName(false)))
-		{
-			ITargetDevicePtr TargetDevice = AcquireDevice();
-
-			if (TargetDevice.IsValid())
-			{
-				FTargetDeviceServicePong* Message = new FTargetDeviceServicePong();
-
-				Message->DeviceID = TargetDevice->GetId().ToString();
-				Message->Name = TargetDevice->GetName();
-				Message->Type = ETargetDeviceTypes::ToString(TargetDevice->GetDeviceType());
-				Message->HostName = FPlatformProcess::ComputerName();
-				Message->HostUser = FPlatformProcess::UserName(false);
-				Message->Connected = TargetDevice->IsConnected();
-				Message->Make = TEXT("@todo");
-				Message->Model = TEXT("@todo");
-				TargetDevice->GetUserCredentials(Message->DeviceUser, Message->DeviceUserPassword);
-				Message->PlatformName = TargetDevice->GetTargetPlatform().PlatformName();
-				Message->Shared = Shared;
-				Message->SupportsMultiLaunch = TargetDevice->SupportsFeature(ETargetDeviceFeatures::MultiLaunch);
-				Message->SupportsPowerOff = TargetDevice->SupportsFeature(ETargetDeviceFeatures::PowerOff);
-				Message->SupportsPowerOn = TargetDevice->SupportsFeature(ETargetDeviceFeatures::PowerOn);
-				Message->SupportsReboot = TargetDevice->SupportsFeature(ETargetDeviceFeatures::Reboot);
-
-				MessageEndpoint->Send(Message, Context->GetSender());
-			}
-		}
-	}
+	void HandlePingMessage(const FTargetDeviceServicePing& InMessage, const IMessageContextRef& Context);
 
 	// Callback for FTargetDeviceServicePowerOff messages.
 	void HandlePowerOffMessage( const FTargetDeviceServicePowerOff& Message, const IMessageContextRef& Context )
@@ -409,8 +358,7 @@ private:
 			return;
 		}
 
-		ITargetDevicePtr TargetDevice = AcquireDevice();
-
+		ITargetDevicePtr TargetDevice = GetDevice(); // Any Device is fine here
 		if (TargetDevice.IsValid())
 		{
 			TargetDevice->PowerOff(Message.Force);
@@ -425,8 +373,7 @@ private:
 			return;
 		}
 
-		ITargetDevicePtr TargetDevice = AcquireDevice();
-
+		ITargetDevicePtr TargetDevice = GetDevice(); // Any Device is fine here
 		if (TargetDevice.IsValid())
 		{
 			TargetDevice->PowerOn();
@@ -441,8 +388,7 @@ private:
 			return;
 		}
 
-		ITargetDevicePtr TargetDevice = AcquireDevice();
-
+		ITargetDevicePtr TargetDevice = GetDevice(); // Any Device is fine here
 		if (TargetDevice.IsValid())
 		{
 			TargetDevice->Reboot();
@@ -450,28 +396,33 @@ private:
 	}
 
 	// Callback for FTargetDeviceServiceRunExecutable messages.
-	void HandleRunExecutableMessage( const FTargetDeviceServiceRunExecutable& Message, const IMessageContextRef& Context )
+	void HandleRunExecutableMessage(const FTargetDeviceServiceRunExecutable& Message, const IMessageContextRef& Context)
 	{
 		if (!Running)
 		{
 			return;
 		}
 
-		ITargetDevicePtr TargetDevice = AcquireDevice();
-
+		ITargetDevicePtr TargetDevice = GetDevice(Message.Variant);
 		if (TargetDevice.IsValid())
 		{
 			uint32 OutProcessId;
 			bool Succeeded = TargetDevice->Run(Message.ExecutablePath, Message.Params, &OutProcessId);
 
-			MessageEndpoint->Send(new FTargetDeviceServiceRunFinished(Message.ExecutablePath, OutProcessId, Succeeded), Context->GetSender());
+			MessageEndpoint->Send(new FTargetDeviceServiceRunFinished(Message.Variant, Message.ExecutablePath, OutProcessId, Succeeded), Context->GetSender());
 		}
 	}
 
 private:
 
 	// Caches the name of the device name that this services exposes.
-	FString CachedDeviceName;
+	FString DeviceName;
+
+	// Caches the platform name of the device name that this services exposes.
+	FName DevicePlatformName;
+
+	// Caches the platform name of the device name that this services exposes.
+	FString DevicePlatformDisplayName;
 
 	// Holds the name of the host that has a claim on the device.
 	FString ClaimHost;
@@ -482,9 +433,6 @@ private:
 	// Holds the name of the user that has a claim on the device.
 	FString ClaimUser;
 
-	// Holds the identifier of the device that this service exposes.
-	const FTargetDeviceId DeviceId;
-
 	// Holds the message endpoint.
 	FMessageEndpointPtr MessageEndpoint;
 
@@ -494,6 +442,9 @@ private:
 	// Holds a flag indicating whether the device is shared with other users.
 	bool Shared;
 
-	// Holds a reference to the target device.
-	ITargetDeviceWeakPtr TargetDevicePtr;
+	// Default target device used when no flavor is specified
+	ITargetDeviceWeakPtr DefaultDevicePtr;
+	
+	// Map of all the Flavors for this Service
+	TMap<FName, ITargetDeviceWeakPtr> TargetDevicePtrs;
 };
