Removed spurious #pragma once

diff --git a/Engine/Source/Developer/TargetDeviceServices/Private/TargetDeviceService.cpp b/Engine/Source/Developer/TargetDeviceServices/Private/TargetDeviceService.cpp
--- a/Engine/Source/Developer/TargetDeviceServices/Private/TargetDeviceService.cpp
+++ b/Engine/Source/Developer/TargetDeviceServices/Private/TargetDeviceService.cpp
@@ -1,7 +1,5 @@
 // Copyright 1998-2014 Epic Games, Inc. All Rights Reserved.
 
-#pragma once
-
 #include "TargetDeviceServicesPrivatePCH.h"
 
 #include "PlatformInfo.h"
