Missed a Username convert which cause devices to not show up.

diff --git a/Engine/Source/Developer/TargetDeviceServices/Private/TargetDeviceService.cpp b/Engine/Source/Developer/TargetDeviceServices/Private/TargetDeviceService.cpp
--- a/Engine/Source/Developer/TargetDeviceServices/Private/TargetDeviceService.cpp
+++ b/Engine/Source/Developer/TargetDeviceServices/Private/TargetDeviceService.cpp
@@ -94,7 +94,7 @@ void FTargetDeviceService::HandlePingMessage(const FTargetDeviceServicePing& InM
 			Message->Name = DefaultDevice->GetName();
 			Message->Type = ETargetDeviceTypes::ToString(DefaultDevice->GetDeviceType());
 			Message->HostName = FPlatformProcess::ComputerName();
-			Message->HostUser = FPlatformProcess::UserName(false);
+			Message->HostUser = FPlatformProcess::UserName(true);
 			Message->Connected = DefaultDevice->IsConnected();
 			Message->Make = TEXT("@todo");
 			Message->Model = TEXT("@todo");
