Restored dummy UClasses, looks like they are still needed (crashes on engine startup)

diff --git a/Engine/Source/Developer/ProfilerMessages/Classes/ProfilerServiceMessages.h b/Engine/Source/Developer/ProfilerMessages/Classes/ProfilerServiceMessages.h
--- a/Engine/Source/Developer/ProfilerMessages/Classes/ProfilerServiceMessages.h
+++ b/Engine/Source/Developer/ProfilerMessages/Classes/ProfilerServiceMessages.h
@@ -564,3 +564,16 @@ struct TStructOpsTypeTraits<FProfilerServiceRequest> : public TStructOpsTypeTrai
 		WithMessageHandling = true
 	};
 };
+
+
+/* Dummy class
+ *****************************************************************************/
+
+UCLASS()
+class UProfilerServiceMessages
+	: public UObject
+{
+public:
+
+	GENERATED_UCLASS_BODY()
+};
\ No newline at end of file
diff --git a/Engine/Source/Developer/ProfilerMessages/Private/ProfilerMessagesModule.cpp b/Engine/Source/Developer/ProfilerMessages/Private/ProfilerMessagesModule.cpp
--- a/Engine/Source/Developer/ProfilerMessages/Private/ProfilerMessagesModule.cpp
+++ b/Engine/Source/Developer/ProfilerMessages/Private/ProfilerMessagesModule.cpp
@@ -21,5 +21,9 @@ public:
 	}
 };
 
+// Dummy class initialization
+UProfilerServiceMessages::UProfilerServiceMessages( const class FPostConstructInitializeProperties& PCIP )
+	: Super(PCIP)
+{ }
 
 IMPLEMENT_MODULE(FProfilerMessagesModule, ProfilerMessages);
diff --git a/Engine/Source/Developer/SessionMessages/Classes/SessionServiceMessages.h b/Engine/Source/Developer/SessionMessages/Classes/SessionServiceMessages.h
--- a/Engine/Source/Developer/SessionMessages/Classes/SessionServiceMessages.h
+++ b/Engine/Source/Developer/SessionMessages/Classes/SessionServiceMessages.h
@@ -183,3 +183,16 @@ struct TStructOpsTypeTraits<FSessionServiceLogUnsubscribe> : public TStructOpsTy
 		WithMessageHandling = true
 	};
 };
+
+
+/* Dummy class
+ *****************************************************************************/
+
+UCLASS()
+class USessionServiceMessages
+	: public UObject
+{
+public:
+
+	GENERATED_UCLASS_BODY()
+};
\ No newline at end of file
diff --git a/Engine/Source/Developer/SessionMessages/Private/SessionMessagesModule.cpp b/Engine/Source/Developer/SessionMessages/Private/SessionMessagesModule.cpp
--- a/Engine/Source/Developer/SessionMessages/Private/SessionMessagesModule.cpp
+++ b/Engine/Source/Developer/SessionMessages/Private/SessionMessagesModule.cpp
@@ -22,5 +22,9 @@ public:
 	}
 };
 
+// Dummy class initialization
+USessionServiceMessages::USessionServiceMessages( const class FPostConstructInitializeProperties& PCIP )
+	: Super(PCIP)
+{ }
 
 IMPLEMENT_MODULE(FSessionMessagesModule, SessionMessages);
diff --git a/Engine/Source/Developer/TargetDeviceServices/Classes/TargetDeviceServiceMessages.h b/Engine/Source/Developer/TargetDeviceServices/Classes/TargetDeviceServiceMessages.h
--- a/Engine/Source/Developer/TargetDeviceServices/Classes/TargetDeviceServiceMessages.h
+++ b/Engine/Source/Developer/TargetDeviceServices/Classes/TargetDeviceServiceMessages.h
@@ -897,3 +897,15 @@ struct TStructOpsTypeTraits<FTargetDeviceServiceRunFinished> : public TStructOps
 		WithMessageHandling = true
 	};
 };
+
+/* Dummy class
+ *****************************************************************************/
+
+UCLASS()
+class UTargetDeviceServiceMessages
+	: public UObject
+{
+public:
+
+	GENERATED_UCLASS_BODY()
+};
\ No newline at end of file
diff --git a/Engine/Source/Developer/TargetDeviceServices/Private/TargetDeviceServicesModule.cpp b/Engine/Source/Developer/TargetDeviceServices/Private/TargetDeviceServicesModule.cpp
--- a/Engine/Source/Developer/TargetDeviceServices/Private/TargetDeviceServicesModule.cpp
+++ b/Engine/Source/Developer/TargetDeviceServices/Private/TargetDeviceServicesModule.cpp
@@ -45,5 +45,9 @@ private:
 	ITargetDeviceServiceManagerPtr DeviceServiceManagerSingleton;
 };
 
+// Dummy class initialization
+UTargetDeviceServiceMessages::UTargetDeviceServiceMessages( const class FPostConstructInitializeProperties& PCIP )
+	: Super(PCIP)
+{ }
 
 IMPLEMENT_MODULE(FTargetDeviceServicesModule, TargetDeviceServices);
diff --git a/Engine/Source/Editor/SwarmInterface/Private/SwarmInterfaceModule.cpp b/Engine/Source/Editor/SwarmInterface/Private/SwarmInterfaceModule.cpp
--- a/Engine/Source/Editor/SwarmInterface/Private/SwarmInterfaceModule.cpp
+++ b/Engine/Source/Editor/SwarmInterface/Private/SwarmInterfaceModule.cpp
@@ -16,5 +16,9 @@ public:
 	virtual void ShutdownModule() override { }
 };
 
+// Dummy class initialization
+USwarmMessages::USwarmMessages( const class FPostConstructInitializeProperties& PCIP )
+	: Super(PCIP)
+{ }
 
 IMPLEMENT_MODULE(FSwarmInterfaceModule, SwarmInterface);
diff --git a/Engine/Source/Editor/SwarmInterface/Private/SwarmMessages.h b/Engine/Source/Editor/SwarmInterface/Private/SwarmMessages.h
--- a/Engine/Source/Editor/SwarmInterface/Private/SwarmMessages.h
+++ b/Engine/Source/Editor/SwarmInterface/Private/SwarmMessages.h
@@ -345,3 +345,15 @@ struct TStructOpsTypeTraits<FSwarmQuitMessage> : public TStructOpsTypeTraitsBase
 		WithMessageHandling = true
 	};
 };
+
+/* Dummy class
+ *****************************************************************************/
+
+UCLASS()
+class USwarmMessages
+	: public UObject
+{
+public:
+
+	GENERATED_UCLASS_BODY()
+};
\ No newline at end of file
diff --git a/Engine/Source/Editor/UnrealEdMessages/Classes/AssetEditorMessages.h b/Engine/Source/Editor/UnrealEdMessages/Classes/AssetEditorMessages.h
--- a/Engine/Source/Editor/UnrealEdMessages/Classes/AssetEditorMessages.h
+++ b/Engine/Source/Editor/UnrealEdMessages/Classes/AssetEditorMessages.h
@@ -42,3 +42,15 @@ struct TStructOpsTypeTraits<FAssetEditorRequestOpenAsset> : public TStructOpsTyp
 		WithMessageHandling = true
 	};
 };
+
+/* Dummy class
+ *****************************************************************************/
+
+UCLASS()
+class UAssetEditorMessages
+	: public UObject
+{
+public:
+
+	GENERATED_UCLASS_BODY()
+};
\ No newline at end of file
diff --git a/Engine/Source/Editor/UnrealEdMessages/Private/UnrealEdMessagesModule.cpp b/Engine/Source/Editor/UnrealEdMessages/Private/UnrealEdMessagesModule.cpp
--- a/Engine/Source/Editor/UnrealEdMessages/Private/UnrealEdMessagesModule.cpp
+++ b/Engine/Source/Editor/UnrealEdMessages/Private/UnrealEdMessagesModule.cpp
@@ -22,5 +22,9 @@ public:
 	}
 };
 
+// Dummy class initialization
+UAssetEditorMessages::UAssetEditorMessages( const class FPostConstructInitializeProperties& PCIP )
+	: Super(PCIP)
+{ }
 
 IMPLEMENT_MODULE(FUnrealEdMessagesModule, UnrealEdMessages);
diff --git a/Engine/Source/Runtime/AutomationMessages/Classes/AutomationWorkerMessages.h b/Engine/Source/Runtime/AutomationMessages/Classes/AutomationWorkerMessages.h
--- a/Engine/Source/Runtime/AutomationMessages/Classes/AutomationWorkerMessages.h
+++ b/Engine/Source/Runtime/AutomationMessages/Classes/AutomationWorkerMessages.h
@@ -451,3 +451,16 @@ struct TStructOpsTypeTraits<FAutomationWorkerScreenImage> : public TStructOpsTyp
 		WithMessageHandling = true
 	};
 };
+
+
+/* Dummy class
+ *****************************************************************************/
+
+UCLASS()
+class UAutomationWorkerMessages
+	: public UObject
+{
+public:
+
+	GENERATED_UCLASS_BODY()
+};
\ No newline at end of file
diff --git a/Engine/Source/Runtime/AutomationMessages/Private/AutomationMessagesModule.cpp b/Engine/Source/Runtime/AutomationMessages/Private/AutomationMessagesModule.cpp
--- a/Engine/Source/Runtime/AutomationMessages/Private/AutomationMessagesModule.cpp
+++ b/Engine/Source/Runtime/AutomationMessages/Private/AutomationMessagesModule.cpp
@@ -21,5 +21,10 @@ public:
 	}
 };
 
+// Dummy class initialization
+UAutomationWorkerMessages::UAutomationWorkerMessages( const class FPostConstructInitializeProperties& PCIP )
+	: Super(PCIP)
+{ }
+
 
 IMPLEMENT_MODULE(FAutomationMessagesModule, AutomationMessages);
diff --git a/Engine/Source/Runtime/EngineMessages/Classes/EngineServiceMessages.h b/Engine/Source/Runtime/EngineMessages/Classes/EngineServiceMessages.h
--- a/Engine/Source/Runtime/EngineMessages/Classes/EngineServiceMessages.h
+++ b/Engine/Source/Runtime/EngineMessages/Classes/EngineServiceMessages.h
@@ -251,3 +251,15 @@ struct TStructOpsTypeTraits<FEngineServiceNotification> : public TStructOpsTypeT
 		WithMessageHandling = true
 	};
 };
+
+/* Dummy class
+ *****************************************************************************/
+
+UCLASS()
+class UEngineServiceMessages
+	: public UObject
+{
+public:
+
+	GENERATED_UCLASS_BODY()
+};
diff --git a/Engine/Source/Runtime/EngineMessages/Private/EngineMessagesModule.cpp b/Engine/Source/Runtime/EngineMessages/Private/EngineMessagesModule.cpp
--- a/Engine/Source/Runtime/EngineMessages/Private/EngineMessagesModule.cpp
+++ b/Engine/Source/Runtime/EngineMessages/Private/EngineMessagesModule.cpp
@@ -22,5 +22,9 @@ public:
 	}
 };
 
+// Dummy class initialization
+UEngineServiceMessages::UEngineServiceMessages( const class FPostConstructInitializeProperties& PCIP )
+	: Super(PCIP)
+{ }
 
 IMPLEMENT_MODULE(FEngineMessagesModule, EngineMessages);
diff --git a/Engine/Source/Runtime/IOS/LaunchDaemonMessages/Classes/IOSMessageProtocol.h b/Engine/Source/Runtime/IOS/LaunchDaemonMessages/Classes/IOSMessageProtocol.h
--- a/Engine/Source/Runtime/IOS/LaunchDaemonMessages/Classes/IOSMessageProtocol.h
+++ b/Engine/Source/Runtime/IOS/LaunchDaemonMessages/Classes/IOSMessageProtocol.h
@@ -99,3 +99,15 @@ struct TStructOpsTypeTraits<FIOSLaunchDaemonLaunchApp> : public TStructOpsTypeTr
 		WithMessageHandling = true
 	};
 };
+
+/* Dummy class
+ *****************************************************************************/
+
+UCLASS()
+class UIOSMessageProtocol
+	: public UObject
+{
+public:
+
+	GENERATED_UCLASS_BODY()
+};
\ No newline at end of file
diff --git a/Engine/Source/Runtime/IOS/LaunchDaemonMessages/Private/LaunchDaemonMessagesModule.cpp b/Engine/Source/Runtime/IOS/LaunchDaemonMessages/Private/LaunchDaemonMessagesModule.cpp
--- a/Engine/Source/Runtime/IOS/LaunchDaemonMessages/Private/LaunchDaemonMessagesModule.cpp
+++ b/Engine/Source/Runtime/IOS/LaunchDaemonMessages/Private/LaunchDaemonMessagesModule.cpp
@@ -21,5 +21,9 @@ public:
 	}
 };
 
+// Dummy class initialization
+UIOSMessageProtocol::UIOSMessageProtocol( const class FPostConstructInitializeProperties& PCIP )
+	: Super(PCIP)
+{ }
 
 IMPLEMENT_MODULE(FLaunchDaemonMessagesModule, LaunchDaemonMessages);
