DeviceManager: Fixed regression: Target devices cannot be released

diff --git a/Engine/Source/Developer/TargetDeviceServices/Private/TargetDeviceService.h b/Engine/Source/Developer/TargetDeviceServices/Private/TargetDeviceService.h
--- a/Engine/Source/Developer/TargetDeviceServices/Private/TargetDeviceService.h
+++ b/Engine/Source/Developer/TargetDeviceServices/Private/TargetDeviceService.h
@@ -42,7 +42,8 @@ public:
 			.Handling<FTargetDeviceServicePowerOff>(this, &FTargetDeviceService::HandlePowerOffMessage)
 			.Handling<FTargetDeviceServicePowerOn>(this, &FTargetDeviceService::HandlePowerOnMessage)
 			.Handling<FTargetDeviceServiceReboot>(this, &FTargetDeviceService::HandleRebootMessage)
-			.Handling<FTargetDeviceServiceRunExecutable>(this, &FTargetDeviceService::HandleRunExecutableMessage);
+			.Handling<FTargetDeviceServiceRunExecutable>(this, &FTargetDeviceService::HandleRunExecutableMessage)
+			.Handling<FTargetDeviceUnclaimed>(this, &FTargetDeviceService::HandleUnclaimedMessage);
 
 		if (MessageEndpoint.IsValid())
 		{
