TargetDeviceServices: Code and documentation cleanup pass

diff --git a/Engine/Source/Developer/TargetDeviceServices/Classes/TargetDeviceServiceMessages.h b/Engine/Source/Developer/TargetDeviceServices/Classes/TargetDeviceServiceMessages.h
--- a/Engine/Source/Developer/TargetDeviceServices/Classes/TargetDeviceServiceMessages.h
+++ b/Engine/Source/Developer/TargetDeviceServices/Classes/TargetDeviceServiceMessages.h
@@ -11,35 +11,26 @@
 /**
  * Implements a message for committing a deployment transaction.
  *
- * @see FTargetDeviceServiceDeployFile
- * @see FTargetDeviceServiceDeployFinished
+ * @see FTargetDeviceServiceDeployFile, FTargetDeviceServiceDeployFinished
  */
 USTRUCT()
 struct FTargetDeviceServiceDeployCommit
 {
 	GENERATED_USTRUCT_BODY()
 
-	/**
-	 * Holds the variant identifier of the target device to use.
-	 */
+	/** Holds the variant identifier of the target device to use. */
 	UPROPERTY()
 	FName Variant;
 
-	/**
-	 * Holds the identifier of the deployment transaction to commit.
-	 */
+	/** Holds the identifier of the deployment transaction to commit. */
 	UPROPERTY()
 	FGuid TransactionId;
 
 
-	/**
-	 * Default constructor.
-	 */
+	/** Default constructor. */
 	FTargetDeviceServiceDeployCommit() { }
 
-	/**
-	 * Creates and initializes a new instance.
-	 */
+	/** Creates and initializes a new instance. */
 	FTargetDeviceServiceDeployCommit(FName InVariant, const FGuid& InTransactionId)
 		: Variant(InVariant)
 		, TransactionId(InTransactionId)
@@ -61,35 +52,26 @@ struct TStructOpsTypeTraits<FTargetDeviceServiceDeployCommit> : public TStructOp
  *
  * The actual file data must be attached to the message.
  *
- * @see FTargetDeviceServiceDeployCommit
- * @see FTargetDeviceServiceDeployFinished
+ * @see FTargetDeviceServiceDeployCommit, FTargetDeviceServiceDeployFinished
  */
 USTRUCT()
 struct FTargetDeviceServiceDeployFile
 {
 	GENERATED_USTRUCT_BODY()
 
-	/**
-	 * Holds the name and path of the file as it will be stored on the target device.
-	 */
+	/** Holds the name and path of the file as it will be stored on the target device. */
 	UPROPERTY()
 	FString TargetFileName;
 
-	/**
-	 * Holds the identifier of the deployment transaction that this file is part of.
-	 */
+	/** Holds the identifier of the deployment transaction that this file is part of. */
 	UPROPERTY()
 	FGuid TransactionId;
 
 
-	/**
-	 * Default constructor.
-	 */
+	/** Default constructor. */
 	FTargetDeviceServiceDeployFile() { }
 
-	/**
-	 * Creates and initializes a new instance.
-	 */
+	/** Creates and initializes a new instance. */
 	FTargetDeviceServiceDeployFile(const FString& InTargetFileName, const FGuid& InTransactionId)
 		: TargetFileName(InTargetFileName)
 		, TransactionId(InTransactionId)
@@ -109,17 +91,14 @@ struct TStructOpsTypeTraits<FTargetDeviceServiceDeployFile> : public TStructOpsT
 /**
  * Implements a message for notifying a target device proxy that a deployment transaction has finished.
  *
- * @see FTargetDeviceServiceDeployFile
- * @see FTargetDeviceServiceDeployCommit
+ * @see FTargetDeviceServiceDeployFile, FTargetDeviceServiceDeployCommit
  */
 USTRUCT()
 struct FTargetDeviceServiceDeployFinished
 {
 	GENERATED_USTRUCT_BODY()
 
-	/**
-	 * Holds the variant identifier of the target device to use.
-	 */
+	/** Holds the variant identifier of the target device to use. */
 	UPROPERTY()
 	FName Variant;
 
@@ -134,27 +113,19 @@ struct FTargetDeviceServiceDeployFinished
 	UPROPERTY()
 	FString AppID;
 
-	/**
-	 * Holds a flag indicating whether the deployment transaction finished successfully.
-	 */
+	/** Holds a flag indicating whether the deployment transaction finished successfully. */
 	UPROPERTY()
 	bool Succeeded;
 
-	/**
-	 * Holds the identifier of the deployment transaction that this file is part of.
-	 */
+	/** Holds the identifier of the deployment transaction that this file is part of. */
 	UPROPERTY()
 	FGuid TransactionId;
 
 
-	/**
-	 * Default constructor.
-	 */
+	/** Default constructor. */
 	FTargetDeviceServiceDeployFinished() { }
 
-	/**
-	 * Creates and initializes a new instance.
-	 */
+	/** Creates and initializes a new instance. */
 	FTargetDeviceServiceDeployFinished(FName InVariant, const FString& InAppId, bool InSucceeded, FGuid InTransactionId)
 		: Variant(InVariant)
 		, AppID(InAppId)
@@ -188,9 +159,7 @@ struct FTargetDeviceServiceLaunchApp
 {
 	GENERATED_USTRUCT_BODY()
 
-	/**
-	 * Holds the variant identifier of the target device to use.
-	 */
+	/** Holds the variant identifier of the target device to use. */
 	UPROPERTY()
 	FName Variant;
 
@@ -205,27 +174,19 @@ struct FTargetDeviceServiceLaunchApp
 	UPROPERTY()
 	FString AppID;
 
-	/**
-	 * The application's build configuration, i.e. Debug or Shipping.
-	 */
+	/** The application's build configuration, i.e. Debug or Shipping. */
 	UPROPERTY()
 	uint8 BuildConfiguration;
 
-	/**
-	 * Holds optional command line parameters for the application.
-	 */
+	/** Holds optional command line parameters for the application. */
 	UPROPERTY()
 	FString Params;
 
 
-	/**
-	 * Default constructor.
-	 */
+	/** Default constructor. */
 	FTargetDeviceServiceLaunchApp() { }
 
-	/**
-	 * Creates and initializes a new instance.
-	 */
+	/** Creates and initializes a new instance. */
 	FTargetDeviceServiceLaunchApp(FName InVariant, const FString& InAppId, uint8 InBuildConfiguration, const FString& InParams)
 		: Variant(InVariant)
 		, AppID(InAppId)
@@ -265,27 +226,19 @@ struct FTargetDeviceServiceLaunchFinished
 	UPROPERTY()
 	FString AppID;
 
-	/**
-	 * Holds the process identifier for the launched application.
-	 */
+	/** Holds the process identifier for the launched application. */
 	UPROPERTY()
 	int32 ProcessId;
 
-	/**
-	 * Holds a flag indicating whether the application was launched successfully.
-	 */
+	/** Holds a flag indicating whether the application was launched successfully. */
 	UPROPERTY()
 	bool Succeeded;
 
 
-	/**
-	 * Default constructor.
-	 */
+	/** Default constructor. */
 	FTargetDeviceServiceLaunchFinished() { }
 
-	/**
-	 * Creates and initializes a new instance.
-	 */
+	/** Creates and initializes a new instance. */
 	FTargetDeviceServiceLaunchFinished(const FString& InAppId, int32 InProcessId, bool InSucceeded)
 		: AppID(InAppId)
 		, ProcessId(InProcessId)
@@ -302,6 +255,7 @@ struct TStructOpsTypeTraits<FTargetDeviceServiceLaunchFinished> : public TStruct
 	};
 };
 
+
 /* Device claiming messages
  *****************************************************************************/
 
@@ -316,32 +270,23 @@ struct FTargetDeviceClaimDenied
 {
 	GENERATED_USTRUCT_BODY()
 
-	/**
-	 * Holds the identifier of the device that is already claimed.
-	 */
+	/** Holds the identifier of the device that is already claimed. */
 	UPROPERTY()
 	FString DeviceName;
 
-	/**
-	 * Holds the name of the host computer that claimed the device.
-	 */
+	/** Holds the name of the host computer that claimed the device. */
 	UPROPERTY()
 	FString HostName;
 
-	/**
-	 * Holds the name of the user that claimed the device.
-	 */
+	/** Holds the name of the user that claimed the device. */
 	UPROPERTY()
 	FString HostUser;
 
-	/**
-	 * Default constructor.
-	 */
+
+	/** Default constructor. */
 	FTargetDeviceClaimDenied( ) { }
 
-	/**
-	 * Creates and initializes a new instance.
-	 */
+	/** Creates and initializes a new instance. */
 	FTargetDeviceClaimDenied(const FString& InDeviceName, const FString& InHostName, const FString& InHostUser)
 		: DeviceName(InDeviceName)
 		, HostName(InHostName)
@@ -370,32 +315,23 @@ struct FTargetDeviceClaimed
 {
 	GENERATED_USTRUCT_BODY()
 
-	/**
-	 * Holds the identifier of the device that is being claimed.
-	 */
+	/** Holds the identifier of the device that is being claimed. */
 	UPROPERTY()
 	FString DeviceName;
 
-	/**
-	 * Holds the name of the host computer that is claiming the device.
-	 */
+	/** Holds the name of the host computer that is claiming the device. */
 	UPROPERTY()
 	FString HostName;
 
-	/**
-	 * Holds the name of the user that is claiming the device.
-	 */
+	/** Holds the name of the user that is claiming the device. */
 	UPROPERTY()
 	FString HostUser;
 
-	/**
-	 * Default constructor.
-	 */
+
+	/** Default constructor. */
 	FTargetDeviceClaimed( ) { }
 
-	/**
-	 * Creates and initializes a new instance.
-	 */
+	/** Creates and initializes a new instance. */
 	FTargetDeviceClaimed(const FString& InDeviceName, const FString& InHostName, const FString& InHostUser)
 		: DeviceName(InDeviceName)
 		, HostName(InHostName)
@@ -416,40 +352,30 @@ struct TStructOpsTypeTraits<FTargetDeviceClaimed> : public TStructOpsTypeTraitsB
 /**
  * Implements a message that is sent when a device is no longer claimed.
  *
- * @see FTargetDeviceClaimDenied
- * @see FTargetDeviceClaimRequest
+ * @see FTargetDeviceClaimDenied, FTargetDeviceClaimRequest
  */
 USTRUCT()
 struct FTargetDeviceUnclaimed
 {
 	GENERATED_USTRUCT_BODY()
 
-	/**
-	 * Holds the identifier of the device that is no longer claimed.
-	 */
+	/** Holds the identifier of the device that is no longer claimed. */
 	UPROPERTY()
 	FString DeviceName;
 
-	/**
-	 * Holds the name of the host computer that had claimed the device.
-	 */
+	/** Holds the name of the host computer that had claimed the device. */
 	UPROPERTY()
 	FString HostName;
 
-	/**
-	 * Holds the name of the user that had claimed the device.
-	 */
+	/** Holds the name of the user that had claimed the device. */
 	UPROPERTY()
 	FString HostUser;
 
-	/**
-	 * Default constructor.
-	 */
+
+	/** Default constructor. */
 	FTargetDeviceUnclaimed( ) { }
 
-	/**
-	 * Creates and initializes a new instance.
-	 */
+	/** Creates and initializes a new instance. */
 	FTargetDeviceUnclaimed(const FString& InDeviceName, const FString& InHostName, const FString& InHostUser)
 		: DeviceName(InDeviceName)
 		, HostName(InHostName)
@@ -478,20 +404,15 @@ struct FTargetDeviceServicePing
 {
 	GENERATED_USTRUCT_BODY()
 
-	/**
-	 * Holds the name of the user who generated the ping.
-	 */
+	/** Holds the name of the user who generated the ping. */
 	UPROPERTY()
 	FString HostUser;
 
-	/**
-	 * Default constructor.
-	 */
+
+	/** Default constructor. */
 	FTargetDeviceServicePing( ) { }
 
-	/**
-	 * Creates and initializes a new instance.
-	 */
+	/** Creates and initializes a new instance. */
 	FTargetDeviceServicePing( const FString& InHostUser )
 		: HostUser(InHostUser)
 	{ }
@@ -552,105 +473,71 @@ struct FTargetDeviceServicePong
 {
 	GENERATED_USTRUCT_BODY()
 
-	/**
-	 * Holds a flag indicating whether the device is currently connected.
-	 */
+	/** Holds a flag indicating whether the device is currently connected. */
 	UPROPERTY()
 	bool Connected;
 
-	/**
-	 * Holds the name of the host computer that the device is attached to.
-	 */
+	/** Holds the name of the host computer that the device is attached to. */
 	UPROPERTY()
 	FString HostName;
 
-	/**
-	 * Holds the name of the user under which the host computer is running.
-	 */
+	/** Holds the name of the user under which the host computer is running. */
 	UPROPERTY()
 	FString HostUser;
 
-	/**
-	 * Holds the make of the device, i.e. Microsoft or Sony.
-	 */
+	/** Holds the make of the device, i.e. Microsoft or Sony. */
 	UPROPERTY()
 	FString Make;
 
-	/**
-	 * Holds the model of the device.
-	 */
+	/** Holds the model of the device. */
 	UPROPERTY()
 	FString Model;
 
-	/**
-	 * Holds the human readable name of the device, i.e "Bob's XBox'.
-	 */
+	/** Holds the human readable name of the device, i.e "Bob's XBox'. */
 	UPROPERTY()
 	FString Name;
 
-	/**
-	 * Holds the name of the user that we log in to remote device as, i.e "root".
-	 */
+	/** Holds the name of the user that we log in to remote device as, i.e "root". */
 	UPROPERTY()
 	FString DeviceUser;
 
-	/**
-	 * Holds the password of the user that we log in to remote device as, i.e "12345".
-	 */
+	/** Holds the password of the user that we log in to remote device as, i.e "12345". */
 	UPROPERTY()
 	FString DeviceUserPassword;
 
-	/**
-	 * Holds a flag indicating whether this device is shared with other users on the network.
-	 */
+	/** Holds a flag indicating whether this device is shared with other users on the network. */
 	UPROPERTY()
 	bool Shared;
 
-	/**
-	 * Holds a flag indicating whether the device supports running multiple application instances in parallel.
-	 */
+	/** Holds a flag indicating whether the device supports running multiple application instances in parallel. */
 	UPROPERTY()
 	bool SupportsMultiLaunch;
 
-	/**
-	 * Holds a flag indicating whether the device can be powered off.
-	 */
+	/** Holds a flag indicating whether the device can be powered off. */
 	UPROPERTY()
 	bool SupportsPowerOff;
 
-	/**
-	 * Holds a flag indicating whether the device can be powered on.
-	 */
+	/** Holds a flag indicating whether the device can be powered on. */
 	UPROPERTY()
 	bool SupportsPowerOn;
 
-	/**
-	 * Holds a flag indicating whether the device can be rebooted.
-	 */
+	/** Holds a flag indicating whether the device can be rebooted. */
 	UPROPERTY()
 	bool SupportsReboot;
 
-	/**
-	* Holds a flag indicating whether the device's target platform supports variants.
-	*/
+	/** Holds a flag indicating whether the device's target platform supports variants. */
 	UPROPERTY()
 	bool SupportsVariants;
 
-	/**
-	 * Holds the device type.
-	 */
+	/** Holds the device type. */
 	UPROPERTY()
 	FString Type;
 
-	/**
-	 * Holds the variant name of the default device.
-	 */
+	/** Holds the variant name of the default device. */
 	UPROPERTY()
 	FName DefaultVariant;
 
-	/**
-	 * List of the Flavors this service supports
-	 */
+	/** List of the Flavors this service supports */
 	UPROPERTY()
 	TArray<FTargetDeviceVariant> Variants;
 };
@@ -684,21 +571,15 @@ struct FTargetDeviceServicePowerOff
 	UPROPERTY()
 	bool Force;
 
-	/**
-	 * Holds the name of the user that wishes to power off the device.
-	 */
+	/** Holds the name of the user that wishes to power off the device. */
 	UPROPERTY()
 	FString Operator;
 
 
-	/**
-	 * Default constructor.
-	 */
+	/** Default constructor. */
 	FTargetDeviceServicePowerOff( ) { }
 
-	/**
-	 * Creates and initializes a new instance.
-	 */
+	/** Creates and initializes a new instance. */
 	FTargetDeviceServicePowerOff( const FString& InOperator, bool InForce )
 		: Force(InForce)
 		, Operator(InOperator)
@@ -723,21 +604,15 @@ struct FTargetDeviceServicePowerOn
 {
 	GENERATED_USTRUCT_BODY()
 
-	/**
-	 * Holds the name of the user that wishes to power on the device.
-	 */
+	/** Holds the name of the user that wishes to power on the device. */
 	UPROPERTY()
 	FString Operator;
 
 
-	/**
-	 * Default constructor.
-	 */
+	/** Default constructor. */
 	FTargetDeviceServicePowerOn( ) { }
 
-	/**
-	 * Creates and initializes a new instance.
-	 */
+	/** Creates and initializes a new instance. */
 	FTargetDeviceServicePowerOn( const FString& InOperator )
 		: Operator(InOperator)
 	{ }
@@ -761,21 +636,15 @@ struct FTargetDeviceServiceReboot
 {
 	GENERATED_USTRUCT_BODY()
 
-	/**
-	 * Holds the name of the user that wishes to reboot the device.
-	 */
+	/** Holds the name of the user that wishes to reboot the device. */
 	UPROPERTY()
 	FString Operator;
 
 
-	/**
-	 * Default constructor.
-	 */
+	/** Default constructor. */
 	FTargetDeviceServiceReboot( ) { }
 
-	/**
-	 * Creates and initializes a new instance.
-	 */
+	/** Creates and initializes a new instance. */
 	FTargetDeviceServiceReboot( const FString& InOperator )
 		: Operator(InOperator)
 	{ }
@@ -803,33 +672,23 @@ struct FTargetDeviceServiceRunExecutable
 {
 	GENERATED_USTRUCT_BODY()
 
-	/**
-	 * Holds the variant identifier of the target device to use for execution.
-	 */
+	/** Holds the variant identifier of the target device to use for execution. */
 	UPROPERTY()
 	FName Variant;
 
-	/**
-	 * Holds the path to the executable on the device.
-	 */
+	/** Holds the path to the executable on the device. */
 	UPROPERTY()
 	FString ExecutablePath;
 
-	/**
-	 * Holds optional command line parameters for the executable.
-	 */
+	/** Holds optional command line parameters for the executable. */
 	UPROPERTY()
 	FString Params;
 
 
-	/**
-	 * Default constructor.
-	 */
+	/** Default constructor. */
 	FTargetDeviceServiceRunExecutable() { }
 
-	/**
-	 * Creates and initializes a new instance.
-	 */
+	/** Creates and initializes a new instance. */
 	FTargetDeviceServiceRunExecutable(FName InVariant, const FString& InExecutablePath, const FString& InParams)
 		: Variant(InVariant)
 		, ExecutablePath(InExecutablePath)
@@ -857,39 +716,27 @@ struct FTargetDeviceServiceRunFinished
 {
 	GENERATED_USTRUCT_BODY()
 
-	/**
-	 * Holds the variant identifier of the target device to use.
-	 */
+	/** Holds the variant identifier of the target device to use. */
 	UPROPERTY()
 	FName Variant;
 
-	/**
-	 * Holds the path to the executable that was run.
-	 */
+	/** Holds the path to the executable that was run. */
 	UPROPERTY()
 	FString ExecutablePath;
 
-	/**
-	 * Holds the process identifier of the running executable.
-	 */
+	/** Holds the process identifier of the running executable. */
 	UPROPERTY()
 	int32 ProcessId;
 
-	/**
-	 * Holds a flag indicating whether the executable started successfully.
-	 */
+	/** Holds a flag indicating whether the executable started successfully. */
 	UPROPERTY()
 	bool Succeeded;
 
 
-	/**
-	 * Default constructor.
-	 */
+	/** Default constructor. */
 	FTargetDeviceServiceRunFinished() { }
 
-	/**
-	 * Creates and initializes a new instance.
-	 */
+	/** Creates and initializes a new instance. */
 	FTargetDeviceServiceRunFinished(FName InVariant, const FString& InExecutablePath, int32 InProcessId, bool InSucceeded)
 		: Variant(InVariant)
 		, ExecutablePath(InExecutablePath)
diff --git a/Engine/Source/Developer/TargetDeviceServices/Private/TargetDeviceProxy.cpp b/Engine/Source/Developer/TargetDeviceServices/Private/TargetDeviceProxy.cpp
--- a/Engine/Source/Developer/TargetDeviceServices/Private/TargetDeviceProxy.cpp
+++ b/Engine/Source/Developer/TargetDeviceServices/Private/TargetDeviceProxy.cpp
@@ -38,42 +38,44 @@ FTargetDeviceProxy::FTargetDeviceProxy(const FString& InName, const FTargetDevic
 
 void FTargetDeviceProxy::UpdateFromMessage( const FTargetDeviceServicePong& Message, const IMessageContextRef& Context )
 {
-	if (Message.Name == Name)
+	if (Message.Name != Name)
 	{
-		MessageAddress = Context->GetSender();
-
-		Connected = Message.Connected;
-		HostName = Message.HostName;
-		HostUser = Message.HostUser;
-		Make = Message.Make;
-		Model = Message.Model;
-		Name = Message.Name;
-		DeviceUser = Message.DeviceUser;
-		DeviceUserPassword = Message.DeviceUserPassword;
-		Shared = Message.Shared;
-		SupportsMultiLaunch = Message.SupportsMultiLaunch;
-		SupportsPowerOff = Message.SupportsPowerOff;
-		SupportsPowerOn = Message.SupportsPowerOn;
-		SupportsReboot = Message.SupportsReboot;
-		SupportsVariants = Message.SupportsVariants;
-		DefaultVariant = Message.DefaultVariant;
-
-		// Update the map of flavors.
-		for (int Index = 0; Index < Message.Variants.Num(); Index++)
-		{
-			const FTargetDeviceVariant& MsgVariant = Message.Variants[Index];
-
-			FTargetDeviceProxyVariant & Variant = TargetDeviceVariants.Add(MsgVariant.VariantName);
-			Variant.DeviceID = MsgVariant.DeviceID;
-			Variant.VariantName = MsgVariant.VariantName;
-			Variant.TargetPlatformName = MsgVariant.TargetPlatformName;
-			Variant.TargetPlatformId = MsgVariant.TargetPlatformId;
-			Variant.VanillaPlatformId = MsgVariant.VanillaPlatformId;
-			Variant.PlatformDisplayName = FText::FromString(MsgVariant.PlatformDisplayName);
-		}
+		return;
+	}
 
-		LastUpdateTime = FDateTime::UtcNow();
+	MessageAddress = Context->GetSender();
+
+	Connected = Message.Connected;
+	HostName = Message.HostName;
+	HostUser = Message.HostUser;
+	Make = Message.Make;
+	Model = Message.Model;
+	Name = Message.Name;
+	DeviceUser = Message.DeviceUser;
+	DeviceUserPassword = Message.DeviceUserPassword;
+	Shared = Message.Shared;
+	SupportsMultiLaunch = Message.SupportsMultiLaunch;
+	SupportsPowerOff = Message.SupportsPowerOff;
+	SupportsPowerOn = Message.SupportsPowerOn;
+	SupportsReboot = Message.SupportsReboot;
+	SupportsVariants = Message.SupportsVariants;
+	DefaultVariant = Message.DefaultVariant;
+
+	// Update the map of flavors.
+	for (int Index = 0; Index < Message.Variants.Num(); Index++)
+	{
+		const FTargetDeviceVariant& MsgVariant = Message.Variants[Index];
+
+		FTargetDeviceProxyVariant & Variant = TargetDeviceVariants.Add(MsgVariant.VariantName);
+		Variant.DeviceID = MsgVariant.DeviceID;
+		Variant.VariantName = MsgVariant.VariantName;
+		Variant.TargetPlatformName = MsgVariant.TargetPlatformName;
+		Variant.TargetPlatformId = MsgVariant.TargetPlatformId;
+		Variant.VanillaPlatformId = MsgVariant.VanillaPlatformId;
+		Variant.PlatformDisplayName = FText::FromString(MsgVariant.PlatformDisplayName);
 	}
+
+	LastUpdateTime = FDateTime::UtcNow();
 }
 
 
@@ -85,45 +87,28 @@ int32 FTargetDeviceProxy::GetNumVariants() const
 	return TargetDeviceVariants.Num();
 }
 
+
 int32 FTargetDeviceProxy::GetVariants(TArray<FName>& OutVariants) const
 {
 	return TargetDeviceVariants.GetKeys(OutVariants);
 }
 
-bool FTargetDeviceProxy::HasVariant(FName InVariant) const
-{
-	if (InVariant == NAME_None)
-	{
-		return TargetDeviceVariants.Contains(DefaultVariant);
-	}
-	return TargetDeviceVariants.Contains(InVariant);
-}
 
 FName FTargetDeviceProxy::GetTargetDeviceVariant(const FString& InDeviceId) const
 {
 	for (TMap<FName, FTargetDeviceProxyVariant>::TConstIterator ItVariant(TargetDeviceVariants); ItVariant; ++ItVariant)
 	{
 		const FTargetDeviceProxyVariant& Variant = ItVariant.Value();
+
 		if (Variant.DeviceID == InDeviceId)
 		{
 			return ItVariant.Key();
 		}
 	}
+
 	return NAME_None;
 }
 
-bool FTargetDeviceProxy::HasDeviceId(const FString& InDeviceId) const
-{
-	for (TMap<FName, FTargetDeviceProxyVariant>::TConstIterator ItVariant(TargetDeviceVariants); ItVariant; ++ItVariant)
-	{
-		const FTargetDeviceProxyVariant& Variant = ItVariant.Value();
-		if (Variant.DeviceID == InDeviceId)
-		{
-			return true;
-		}
-	}
-	return false;
-}
 
 const FString& FTargetDeviceProxy::GetTargetDeviceId(FName InVariant) const
 {
@@ -131,21 +116,10 @@ const FString& FTargetDeviceProxy::GetTargetDeviceId(FName InVariant) const
 	{
 		return TargetDeviceVariants[DefaultVariant].DeviceID;
 	}
+
 	return TargetDeviceVariants[InVariant].DeviceID;
 }
 
-bool FTargetDeviceProxy::HasTargetPlatform(FName InTargetPlatformId) const
-{
-	for (TMap<FName, FTargetDeviceProxyVariant>::TConstIterator ItVariant(TargetDeviceVariants); ItVariant; ++ItVariant)
-	{
-		const FTargetDeviceProxyVariant& Variant = ItVariant.Value();
-		if (Variant.TargetPlatformId == InTargetPlatformId)
-		{
-			return true;
-		}
-	}
-	return false;
-}
 
 FString FTargetDeviceProxy::GetTargetPlatformName(FName InVariant) const
 {
@@ -153,36 +127,87 @@ FString FTargetDeviceProxy::GetTargetPlatformName(FName InVariant) const
 	{
 		return TargetDeviceVariants[DefaultVariant].TargetPlatformName;
 	}
+
 	return TargetDeviceVariants[InVariant].TargetPlatformName;
 }
 
+
 FName FTargetDeviceProxy::GetTargetPlatformId(FName InVariant) const
 {
 	if (InVariant == NAME_None)
 	{
 		return TargetDeviceVariants[DefaultVariant].TargetPlatformId;
 	}
+
 	return TargetDeviceVariants[InVariant].TargetPlatformId;
 }
 
+
 FName FTargetDeviceProxy::GetVanillaPlatformId(FName InVariant) const
 {
 	if (InVariant == NAME_None)
 	{
 		return TargetDeviceVariants[DefaultVariant].VanillaPlatformId;
 	}
+
 	return TargetDeviceVariants[InVariant].VanillaPlatformId;
 }
 
+
 FText FTargetDeviceProxy::GetPlatformDisplayName(FName InVariant) const
 {
 	if (InVariant == NAME_None)
 	{
 		return TargetDeviceVariants[DefaultVariant].PlatformDisplayName;
 	}
+
 	return TargetDeviceVariants[InVariant].PlatformDisplayName;
 }
 
+
+bool FTargetDeviceProxy::HasDeviceId(const FString& InDeviceId) const
+{
+	for (TMap<FName, FTargetDeviceProxyVariant>::TConstIterator ItVariant(TargetDeviceVariants); ItVariant; ++ItVariant)
+	{
+		const FTargetDeviceProxyVariant& Variant = ItVariant.Value();
+
+		if (Variant.DeviceID == InDeviceId)
+		{
+			return true;
+		}
+	}
+
+	return false;
+}
+
+
+bool FTargetDeviceProxy::HasTargetPlatform(FName InTargetPlatformId) const
+{
+	for (TMap<FName, FTargetDeviceProxyVariant>::TConstIterator ItVariant(TargetDeviceVariants); ItVariant; ++ItVariant)
+	{
+		const FTargetDeviceProxyVariant& Variant = ItVariant.Value();
+
+		if (Variant.TargetPlatformId == InTargetPlatformId)
+		{
+			return true;
+		}
+	}
+
+	return false;
+}
+
+
+bool FTargetDeviceProxy::HasVariant(FName InVariant) const
+{
+	if (InVariant == NAME_None)
+	{
+		return TargetDeviceVariants.Contains(DefaultVariant);
+	}
+
+	return TargetDeviceVariants.Contains(InVariant);
+}
+
+
 bool FTargetDeviceProxy::DeployApp(FName InVariant, const TMap<FString, FString>& Files, const FGuid& TransactionId)
 {
 	for (TMap<FString, FString>::TConstIterator It(Files); It; ++It)
@@ -198,6 +223,7 @@ bool FTargetDeviceProxy::DeployApp(FName InVariant, const TMap<FString, FString>
 	return true;
 }
 
+
 bool FTargetDeviceProxy::LaunchApp(FName InVariant, const FString& AppId, EBuildConfigurations::Type BuildConfiguration, const FString& Params)
 {
 	MessageEndpoint->Send(new FTargetDeviceServiceLaunchApp(InVariant, AppId, BuildConfiguration, Params), MessageAddress);
@@ -205,22 +231,25 @@ bool FTargetDeviceProxy::LaunchApp(FName InVariant, const FString& AppId, EBuild
 	return true;
 }
 
-void FTargetDeviceProxy::PowerOff( bool Force )
+
+void FTargetDeviceProxy::PowerOff(bool Force)
 {
 	MessageEndpoint->Send(new FTargetDeviceServicePowerOff(FPlatformProcess::UserName(true), Force), MessageAddress);
 }
 
-void FTargetDeviceProxy::PowerOn( )
+
+void FTargetDeviceProxy::PowerOn()
 {
 	MessageEndpoint->Send(new FTargetDeviceServicePowerOn(FPlatformProcess::UserName(true)), MessageAddress);
 }
 
 
-void FTargetDeviceProxy::Reboot( )
+void FTargetDeviceProxy::Reboot()
 {
 	MessageEndpoint->Send(new FTargetDeviceServiceReboot(FPlatformProcess::UserName(true)), MessageAddress);
 }
 
+
 void FTargetDeviceProxy::Run(FName InVariant, const FString& ExecutablePath, const FString& Params)
 {
 	MessageEndpoint->Send(new FTargetDeviceServiceRunExecutable(InVariant, ExecutablePath, Params), MessageAddress);
@@ -230,7 +259,7 @@ void FTargetDeviceProxy::Run(FName InVariant, const FString& ExecutablePath, con
 /* FTargetDeviceProxy implementation
  *****************************************************************************/
 
-void FTargetDeviceProxy::InitializeMessaging( )
+void FTargetDeviceProxy::InitializeMessaging()
 {
 	MessageEndpoint = FMessageEndpoint::Builder(FName(*FString::Printf(TEXT("FTargetDeviceProxy (%s)"), *Name)))
 		.Handling<FTargetDeviceServiceDeployFinished>(this, &FTargetDeviceProxy::HandleDeployFinishedMessage)
diff --git a/Engine/Source/Developer/TargetDeviceServices/Private/TargetDeviceProxy.h b/Engine/Source/Developer/TargetDeviceServices/Private/TargetDeviceProxy.h
--- a/Engine/Source/Developer/TargetDeviceServices/Private/TargetDeviceProxy.h
+++ b/Engine/Source/Developer/TargetDeviceServices/Private/TargetDeviceProxy.h
@@ -18,9 +18,7 @@ class FTargetDeviceProxy
 {
 public:
 
-	/**
-	 * Default constructor
-	 */
+	/** Default constructor. */
 	FTargetDeviceProxy() { }
 
 	/**
@@ -89,25 +87,12 @@ public:
 	}
 
 	virtual int32 GetNumVariants() const override;
-
 	virtual int32 GetVariants(TArray<FName>& OutVariants) const override;
-
-	virtual bool HasVariant(FName InVariant) const override;
-
 	virtual FName GetTargetDeviceVariant(const FString& InDeviceId) const override;
-
-	virtual bool HasDeviceId(const FString& InDeviceId) const override;
-
 	virtual const FString& GetTargetDeviceId(FName InVariant) const override;
-
-	virtual bool HasTargetPlatform(FName InTargetPlatformId) const override;
-
 	virtual FString GetTargetPlatformName(FName InVariant) const override;
-
 	virtual FName GetTargetPlatformId(FName InVariant) const override;
-
 	virtual FName GetVanillaPlatformId(FName InVariant) const override;
-
 	virtual FText GetPlatformDisplayName(FName InVariant) const override;
 
 	virtual const FString& GetHostName() const override
@@ -150,6 +135,10 @@ public:
 		return Type;
 	}
 
+	virtual bool HasDeviceId(const FString& InDeviceId) const override;
+	virtual bool HasVariant(FName InVariant) const override;
+	virtual bool HasTargetPlatform(FName InTargetPlatformId) const override;
+
 	virtual bool IsConnected() const override
 	{
 		return Connected;
@@ -161,7 +150,6 @@ public:
 	}
 
 	virtual bool DeployApp(FName InVariant, const TMap<FString, FString>& Files, const FGuid& TransactionId) override;
-
 	virtual bool LaunchApp(FName InVariant, const FString& AppId, EBuildConfigurations::Type BuildConfiguration, const FString& Params) override;
 
 	virtual FOnTargetDeviceProxyDeployCommitted& OnDeployCommitted() override
@@ -180,119 +168,115 @@ public:
 	}
 
 	virtual void PowerOff(bool Force) override;
-
 	virtual void PowerOn() override;
-
 	virtual void Reboot() override;
-
 	virtual void Run(FName InVariant, const FString& ExecutablePath, const FString& Params) override;
 
-
 protected:
 
-	/**
-	 * Initializes the message endpoint.
-	 */
+	/** Initializes the message endpoint. */
 	void InitializeMessaging();
 
 private:
 
-	// Handles FTargetDeviceServiceDeployFinishedMessage messages.
+	/** Handles FTargetDeviceServiceDeployFinishedMessage messages. */
 	void HandleDeployFinishedMessage(const FTargetDeviceServiceDeployFinished& Message, const IMessageContextRef& Context);
 
-	// Handles FTargetDeviceServiceLaunchFinishedMessage messages.
+	/** Handles FTargetDeviceServiceLaunchFinishedMessage messages. */
 	void HandleLaunchFinishedMessage(const FTargetDeviceServiceLaunchFinished& Message, const IMessageContextRef& Context);
 
 private:
 
-	// Holds a flag indicating whether the device is connected.
+	/** Holds a flag indicating whether the device is connected. */
 	bool Connected;
 
-	// Holds the name of the computer that hosts the device.
+	/** Holds the name of the computer that hosts the device. */
 	FString HostName;
 
-	// Holds the name of the user that owns the device.
+	/** Holds the name of the user that owns the device. */
 	FString HostUser;
 
-	// Holds the time at which the last ping reply was received.
+	/** Holds the time at which the last ping reply was received. */
 	FDateTime LastUpdateTime;
 
-	// Holds the device make.
+	/** Holds the device make. */
 	FString Make;
 
-	// Holds the remote device's message bus address.
+	/** Holds the remote device's message bus address. */
 	FMessageAddress MessageAddress;
 
-	// Holds the local message bus endpoint.
+	/** Holds the local message bus endpoint. */
 	FMessageEndpointPtr MessageEndpoint;
 
-	// Holds the remote device's model.
+	/** Holds the remote device's model. */
 	FString Model;
 
-	// Holds the name of the device.
+	/** Holds the name of the device. */
 	FString Name;
 
-	// Holds device user
+	/** Holds device user. */
 	FString DeviceUser;
 
-	// Holds device user password
+	/** Holds device user password. */
 	FString DeviceUserPassword;
 
-	// Holds a flag indicating whether the device is being shared with other users.
+	/** Holds a flag indicating whether the device is being shared with other users. */
 	bool Shared;
 
-	// Holds a flag indicating whether the device supports multi-launch.
+	/** Holds a flag indicating whether the device supports multi-launch. */
 	bool SupportsMultiLaunch;
 
-	// Holds a flag indicating whether the device can power off.
+	/** Holds a flag indicating whether the device can power off. */
 	bool SupportsPowerOff;
 
-	// Holds a flag indicating whether the device can be power on.
+	/** Holds a flag indicating whether the device can be power on. */
 	bool SupportsPowerOn;
 
-	// Holds a flag indicating whether the device can reboot.
+	/** Holds a flag indicating whether the device can reboot. */
 	bool SupportsReboot;
 
-	// Holds a flag indicating whether the device's target platform supports variants.
+	/** Holds a flag indicating whether the device's target platform supports variants. */
 	bool SupportsVariants;
 
-	// Holds the device type.
+	/** Holds the device type. */
 	FString Type;
 
-	// Holds default variant name.
+	/** Holds default variant name. */
 	FName DefaultVariant;
 
-	// Holds data about a device proxy variant
+
+	/** Holds data about a device proxy variant. */
 	class FTargetDeviceProxyVariant
 	{
 	public:
-		// Holds a string version of the variants device id.
+
+		/** Holds a string version of the variants device id. */
 		FString DeviceID;
 
-		// Holds the variant name, this is the the map key as well
+		/** Holds the variant name, this is the the map key as well. */
 		FName VariantName;
 
-		// Platform information
+		/** Platform information. */
 		FString TargetPlatformName;
 		FName TargetPlatformId;
 		FName VanillaPlatformId;
 		FText PlatformDisplayName;
 	};
 
-	// Map of all the Variants for this Device
+	/** Map of all the Variants for this Device. */
 	TMap<FName, FTargetDeviceProxyVariant> TargetDeviceVariants;
 
 private:
 
-	// Holds a delegate to be invoked when a build has been deployed to the target device.
+	/** Holds a delegate to be invoked when a build has been deployed to the target device. */
 	FOnTargetDeviceProxyDeployCommitted DeployCommittedDelegate;
 
-	// Holds a delegate to be invoked when a build has failed to deploy to the target device.
+	/** Holds a delegate to be invoked when a build has failed to deploy to the target device. */
 	FOnTargetDeviceProxyDeployFailed DeployFailedDelegate;
 
-	// Holds a delegate to be invoked when a build has failed to launch on the target device.
+	/** Holds a delegate to be invoked when a build has failed to launch on the target device. */
 	FOnTargetDeviceProxyLaunchFailed LaunchFailedDelegate;
 
-	// Holds a delegate to be invoked when a build has succeeded to launch on the target device.
+	/** Holds a delegate to be invoked when a build has succeeded to launch on the target device. */
 	FOnTargetDeviceProxyLaunchSucceeded LaunchSucceededDelegate;
 };
diff --git a/Engine/Source/Developer/TargetDeviceServices/Private/TargetDeviceProxyManager.cpp b/Engine/Source/Developer/TargetDeviceServices/Private/TargetDeviceProxyManager.cpp
--- a/Engine/Source/Developer/TargetDeviceServices/Private/TargetDeviceProxyManager.cpp
+++ b/Engine/Source/Developer/TargetDeviceServices/Private/TargetDeviceProxyManager.cpp
@@ -6,7 +6,7 @@
 /* FTargetDeviceProxyManager structors
  *****************************************************************************/
 
-FTargetDeviceProxyManager::FTargetDeviceProxyManager( )
+FTargetDeviceProxyManager::FTargetDeviceProxyManager()
 {
 	MessageEndpoint = FMessageEndpoint::Builder("FTargetDeviceProxyManager")
 		.Handling<FTargetDeviceServicePong>(this, &FTargetDeviceProxyManager::HandlePongMessage);
@@ -21,7 +21,7 @@ FTargetDeviceProxyManager::FTargetDeviceProxyManager( )
 }
 
 
-FTargetDeviceProxyManager::~FTargetDeviceProxyManager( )
+FTargetDeviceProxyManager::~FTargetDeviceProxyManager()
 {
 	FTicker::GetCoreTicker().RemoveTicker(TickDelegate);
 }
@@ -30,7 +30,7 @@ FTargetDeviceProxyManager::~FTargetDeviceProxyManager( )
 /* ITargetDeviceProxyLocator interface
  *****************************************************************************/
 
-ITargetDeviceProxyPtr FTargetDeviceProxyManager::FindProxy( const FString& Name ) 
+ITargetDeviceProxyPtr FTargetDeviceProxyManager::FindProxy(const FString& Name) 
 {
 	return Proxies.FindRef(Name);
 }
@@ -50,6 +50,7 @@ ITargetDeviceProxyRef FTargetDeviceProxyManager::FindOrAddProxy(const FString& N
 	return Proxy.ToSharedRef();
 }
 
+
 ITargetDeviceProxyPtr FTargetDeviceProxyManager::FindProxyDeviceForTargetDevice(const FString& DeviceId)
 {
 	for (TMap<FString, TSharedPtr<FTargetDeviceProxy> >::TConstIterator ItProxies(Proxies); ItProxies; ++ItProxies)
@@ -65,6 +66,7 @@ ITargetDeviceProxyPtr FTargetDeviceProxyManager::FindProxyDeviceForTargetDevice(
 	return ITargetDeviceProxyPtr();
 }
 
+
 void FTargetDeviceProxyManager::GetProxies(FName TargetPlatformName, bool IncludeUnshared, TArray<ITargetDeviceProxyPtr>& OutProxies)
 {
 	OutProxies.Reset();
@@ -86,7 +88,7 @@ void FTargetDeviceProxyManager::GetProxies(FName TargetPlatformName, bool Includ
 /* FTargetDeviceProxyManager implementation
  *****************************************************************************/
 
-void FTargetDeviceProxyManager::RemoveDeadProxies( )
+void FTargetDeviceProxyManager::RemoveDeadProxies()
 {
 	FDateTime CurrentTime = FDateTime::UtcNow();
 
@@ -111,10 +113,10 @@ void FTargetDeviceProxyManager::SendPing()
 }
 
 
-/* FTargetDeviceProxyManager event handlers
+/* FTargetDeviceProxyManager callbacks
  *****************************************************************************/
 
-void FTargetDeviceProxyManager::HandlePongMessage( const FTargetDeviceServicePong& Message, const IMessageContextRef& Context )
+void FTargetDeviceProxyManager::HandlePongMessage(const FTargetDeviceServicePong& Message, const IMessageContextRef& Context)
 {
 	TSharedPtr<FTargetDeviceProxy>& Proxy = Proxies.FindOrAdd(Message.Name);
 
@@ -130,7 +132,7 @@ void FTargetDeviceProxyManager::HandlePongMessage( const FTargetDeviceServicePon
 }
 
 
-bool FTargetDeviceProxyManager::HandleTicker( float DeltaTime )
+bool FTargetDeviceProxyManager::HandleTicker(float DeltaTime)
 {
 	RemoveDeadProxies();
 	SendPing();
diff --git a/Engine/Source/Developer/TargetDeviceServices/Private/TargetDeviceProxyManager.h b/Engine/Source/Developer/TargetDeviceServices/Private/TargetDeviceProxyManager.h
--- a/Engine/Source/Developer/TargetDeviceServices/Private/TargetDeviceProxyManager.h
+++ b/Engine/Source/Developer/TargetDeviceServices/Private/TargetDeviceProxyManager.h
@@ -11,15 +11,11 @@ class FTargetDeviceProxyManager
 {
 public:
 
-	/**
-	 * Default constructor.
-	 */
-	FTargetDeviceProxyManager( );
+	/** Default constructor. */
+	FTargetDeviceProxyManager();
 
-	/**
-	 * Destructor.
-	 */
-	virtual ~FTargetDeviceProxyManager( );
+	/** Destructor. */
+	virtual ~FTargetDeviceProxyManager();
 
 public:
 
@@ -33,53 +29,49 @@ public:
 	virtual void GetProxies(FName TargetPlatformName, bool IncludeUnshared, TArray<ITargetDeviceProxyPtr>& OutProxies) override;
 
 	DECLARE_DERIVED_EVENT(FTargetDeviceProxyManager, ITargetDeviceProxyManager::FOnTargetDeviceProxyAdded, FOnTargetDeviceProxyAdded);
-	virtual FOnTargetDeviceProxyAdded& OnProxyAdded( ) override
+	virtual FOnTargetDeviceProxyAdded& OnProxyAdded() override
 	{
 		return ProxyAddedDelegate;
 	}
 
 	DECLARE_DERIVED_EVENT(FTargetDeviceProxyManager, ITargetDeviceProxyManager::FOnTargetDeviceProxyRemoved, FOnTargetDeviceProxyRemoved);
-	virtual FOnTargetDeviceProxyRemoved& OnProxyRemoved( ) override
+	virtual FOnTargetDeviceProxyRemoved& OnProxyRemoved() override
 	{
 		return ProxyRemovedDelegate;
 	}
 
 protected:
 
-	/**
-	 * Removes all target device proxies that timed out.
-	 */
-	void RemoveDeadProxies( );
+	/** Removes all target device proxies that timed out. */
+	void RemoveDeadProxies();
 
-	/**
-	 * Pings all target devices on the network.
-	 */
-	void SendPing( );
+	/** Pings all target devices on the network. */
+	void SendPing();
 
 private:
 	
-	// Handles FTargetDeviceServicePong messages.
-	void HandlePongMessage( const FTargetDeviceServicePong& Message, const IMessageContextRef& Context );
+	/** Handles FTargetDeviceServicePong messages. */
+	void HandlePongMessage(const FTargetDeviceServicePong& Message, const IMessageContextRef& Context);
 
-	// Handles ticks from the ticker.
-	bool HandleTicker( float DeltaTime );
+	/** Handles ticks from the ticker. */
+	bool HandleTicker(float DeltaTime);
 
 private:
 
-	// Holds the message endpoint.
+	/** Holds the message endpoint. */
 	FMessageEndpointPtr MessageEndpoint;
 
-	// Holds the collection of proxies.
+	/** Holds the collection of proxies. */
 	TMap<FString, FTargetDeviceProxyPtr> Proxies;
 
 private:
 
-	// Holds a delegate that is invoked when a target device proxy has been added.
+	/** Holds a delegate that is invoked when a target device proxy has been added. */
 	FOnTargetDeviceProxyAdded ProxyAddedDelegate;
 
-	// Holds a delegate that is invoked when a target device proxy has been removed.
+	/** Holds a delegate that is invoked when a target device proxy has been removed. */
 	FOnTargetDeviceProxyRemoved ProxyRemovedDelegate;
 
-	// Holds a delegate to be invoked when the widget ticks.
+	/** Holds a delegate to be invoked when the widget ticks. */
 	FTickerDelegate TickDelegate;
 };
diff --git a/Engine/Source/Developer/TargetDeviceServices/Private/TargetDeviceService.cpp b/Engine/Source/Developer/TargetDeviceServices/Private/TargetDeviceService.cpp
--- a/Engine/Source/Developer/TargetDeviceServices/Private/TargetDeviceService.cpp
+++ b/Engine/Source/Developer/TargetDeviceServices/Private/TargetDeviceService.cpp
@@ -1,9 +1,65 @@
 // Copyright 1998-2014 Epic Games, Inc. All Rights Reserved.
 
 #include "TargetDeviceServicesPrivatePCH.h"
-
 #include "PlatformInfo.h"
 
+
+/* Local helpers
+ *****************************************************************************/
+
+struct FVariantSortCallback
+{
+	FORCEINLINE bool operator()(const ITargetDeviceWeakPtr& A, const ITargetDeviceWeakPtr& B) const
+	{
+		ITargetDevicePtr APtr = A.Pin();
+		ITargetDevicePtr BPtr = B.Pin();
+
+		return APtr->GetTargetPlatform().GetVariantPriority() > BPtr->GetTargetPlatform().GetVariantPriority();
+	}
+};
+
+
+/* FTargetDeviceService structors
+ *****************************************************************************/
+
+FTargetDeviceService::FTargetDeviceService(const FString& InDeviceName, const IMessageBusRef& InMessageBus)
+	: DeviceName(InDeviceName)
+	, Running(false)
+	, Shared(false)
+{
+	// initialize messaging
+	MessageEndpoint = FMessageEndpoint::Builder(FName(*FString::Printf(TEXT("FTargetDeviceService (%s)"), *DeviceName)), InMessageBus)
+		.Handling<FTargetDeviceClaimDenied>(this, &FTargetDeviceService::HandleClaimDeniedMessage)
+		.Handling<FTargetDeviceClaimed>(this, &FTargetDeviceService::HandleClaimedMessage)
+		.Handling<FTargetDeviceServiceDeployCommit>(this, &FTargetDeviceService::HandleDeployCommitMessage)
+		.Handling<FTargetDeviceServiceDeployFile>(this, &FTargetDeviceService::HandleDeployFileMessage)
+		.Handling<FTargetDeviceServiceLaunchApp>(this, &FTargetDeviceService::HandleLaunchAppMessage)
+		.Handling<FTargetDeviceServicePing>(this, &FTargetDeviceService::HandlePingMessage)
+		.Handling<FTargetDeviceServicePowerOff>(this, &FTargetDeviceService::HandlePowerOffMessage)
+		.Handling<FTargetDeviceServicePowerOn>(this, &FTargetDeviceService::HandlePowerOnMessage)
+		.Handling<FTargetDeviceServiceReboot>(this, &FTargetDeviceService::HandleRebootMessage)
+		.Handling<FTargetDeviceServiceRunExecutable>(this, &FTargetDeviceService::HandleRunExecutableMessage)
+		.Handling<FTargetDeviceUnclaimed>(this, &FTargetDeviceService::HandleUnclaimedMessage);
+
+	if (MessageEndpoint.IsValid())
+	{
+		MessageEndpoint->Subscribe<FTargetDeviceClaimed>();
+		MessageEndpoint->Subscribe<FTargetDeviceUnclaimed>();
+		MessageEndpoint->Subscribe<FTargetDeviceServicePing>();
+	}
+}
+
+
+FTargetDeviceService::~FTargetDeviceService()
+{
+	Stop();
+	FMessageEndpoint::SafeRelease(MessageEndpoint);
+}
+
+
+/* ITargetDeviceService interface
+ *****************************************************************************/
+
 void FTargetDeviceService::AddTargetDevice(ITargetDevicePtr InDevice)
 {
 	if (!InDevice.IsValid())
@@ -31,15 +87,16 @@ void FTargetDeviceService::AddTargetDevice(ITargetDevicePtr InDevice)
 		DevicePlatformDisplayName = bSplit ? Left.Trim() : Full;
 	}
 
-	// Make sure we don't have conflicting flavors for the same device!
-	// Could also be a double add, which ideally would be avoided.
+	// make sure we don't have conflicting flavors for the same device!
+	// could also be a double add, which ideally would be avoided.
 	check(!(TargetDevicePtrs.FindRef(Variant).IsValid()));
 
 	TargetDevicePtrs.Add(Variant, InDevice);
 
-	// Sort and choose cache the default
+	// sort and choose cache the default
 	TargetDevicePtrs.ValueSort(FVariantSortCallback());
 	auto DeviceIterator = TargetDevicePtrs.CreateIterator();
+
 	if (DeviceIterator)
 	{
 		DefaultDevicePtr = (*DeviceIterator).Value;
@@ -50,6 +107,42 @@ void FTargetDeviceService::AddTargetDevice(ITargetDevicePtr InDevice)
 	}
 }
 
+
+bool FTargetDeviceService::CanStart(FName InFlavor) const
+{
+	ITargetDevicePtr TargetDevice = GetDevice(InFlavor);
+
+	if (TargetDevice.IsValid())
+	{
+		return TargetDevice->IsConnected();
+	}
+
+	return false;
+}
+
+
+ITargetDevicePtr FTargetDeviceService::GetDevice(FName InVariant) const
+{
+	ITargetDevicePtr TargetDevice;
+
+	if (InVariant == NAME_None)
+	{
+		TargetDevice = DefaultDevicePtr.Pin();
+	}
+	else
+	{
+		const ITargetDeviceWeakPtr * WeakTargetDevicePtr = TargetDevicePtrs.Find(InVariant);
+
+		if (WeakTargetDevicePtr != nullptr)
+		{
+			TargetDevice = WeakTargetDevicePtr->Pin();
+		}
+	}
+
+	return TargetDevice;
+}
+
+
 void FTargetDeviceService::RemoveTargetDevice(ITargetDevicePtr InDevice)
 {
 	if (!InDevice.IsValid())
@@ -63,6 +156,7 @@ void FTargetDeviceService::RemoveTargetDevice(ITargetDevicePtr InDevice)
 
 	// Cache the default
 	auto DeviceIterator = TargetDevicePtrs.CreateIterator();
+
 	if (DeviceIterator)
 	{
 		DefaultDevicePtr = (*DeviceIterator).Value;
@@ -73,6 +167,212 @@ void FTargetDeviceService::RemoveTargetDevice(ITargetDevicePtr InDevice)
 	}
 }
 
+
+bool FTargetDeviceService::Start()
+{
+	if (!Running && MessageEndpoint.IsValid())
+	{
+		// notify other services
+		ClaimAddress = MessageEndpoint->GetAddress();
+		ClaimHost = FPlatformProcess::ComputerName();
+		ClaimUser = FPlatformProcess::UserName(true);
+
+		MessageEndpoint->Publish(new FTargetDeviceClaimed(DeviceName, ClaimHost, ClaimUser));
+
+		Running = true;
+	}
+
+	return true;
+}
+
+
+void FTargetDeviceService::Stop()
+{
+	if (Running)
+	{
+		// notify other services
+		MessageEndpoint->Publish(new FTargetDeviceUnclaimed(DeviceName, FPlatformProcess::ComputerName(), FPlatformProcess::UserName(true)));
+
+		Running = false;
+	}
+}
+
+
+/* FTargetDeviceService implementation
+ *****************************************************************************/
+
+bool FTargetDeviceService::StoreDeployedFile(FArchive* FileReader, const FString& TargetFileName) const
+{
+	if (FileReader == nullptr)
+	{
+		return false;
+	}
+
+	// create target file
+	FArchive* FileWriter = IFileManager::Get().CreateFileWriter(*TargetFileName);
+
+	if (FileWriter == nullptr)
+	{
+		return false;
+	}
+
+	FileReader->Seek(0);
+
+	// copy file contents
+	int64 BytesRemaining = FileReader->TotalSize();
+	int32 BufferSize = 128 * 1024;
+
+	if (BytesRemaining < BufferSize)
+	{
+		BufferSize = BytesRemaining;
+	}
+
+	void* Buffer = FMemory::Malloc(BufferSize);
+
+	while (BytesRemaining > 0)
+	{
+		FileReader->Serialize(Buffer, BufferSize);
+		FileWriter->Serialize(Buffer, BufferSize);
+
+		BytesRemaining -= BufferSize;
+
+		if (BytesRemaining < BufferSize)
+		{
+			BufferSize = BytesRemaining;
+		}
+	}
+
+	// clean up
+	FMemory::Free(Buffer);
+	delete FileWriter;
+
+	return true;
+}
+
+
+/* FTargetDeviceService callbacks
+ *****************************************************************************/
+
+void FTargetDeviceService::HandleClaimDeniedMessage(const FTargetDeviceClaimDenied& Message, const IMessageContextRef& Context)
+{
+	if (Running && (Message.DeviceName == DeviceName))
+	{
+		Stop();
+
+		ClaimAddress = Context->GetSender();
+		ClaimHost = Message.HostName;
+		ClaimUser = Message.HostUser;
+	}
+}
+
+
+void FTargetDeviceService::HandleClaimedMessage(const FTargetDeviceClaimed& Message, const IMessageContextRef& Context)
+{
+	if (Message.DeviceName != DeviceName)
+	{
+		return;
+	}
+
+	if (Running)
+	{
+		if (Context->GetSender() != MessageEndpoint->GetAddress())
+		{
+			MessageEndpoint->Send(new FTargetDeviceClaimDenied(DeviceName, FPlatformProcess::ComputerName(), FPlatformProcess::UserName(true)), Context->GetSender());
+		}
+	}
+	else
+	{
+		ClaimAddress = Context->GetSender();
+		ClaimHost = Message.HostName;
+		ClaimUser = Message.HostUser;
+	}
+}
+
+
+void FTargetDeviceService::HandleUnclaimedMessage(const FTargetDeviceUnclaimed& Message, const IMessageContextRef& Context)
+{
+	if (Message.DeviceName == DeviceName)
+	{
+		if (Context->GetSender() == ClaimAddress)
+		{
+			ClaimAddress = FGuid();
+			ClaimHost.Empty();
+			ClaimUser.Empty();
+		}
+	}
+}
+
+
+void FTargetDeviceService::HandleDeployFileMessage(const FTargetDeviceServiceDeployFile& Message, const IMessageContextRef& Context)
+{
+	if (!Running)
+	{
+		return;
+	}
+
+	IMessageAttachmentPtr Attachment = Context->GetAttachment();
+
+	if (Attachment.IsValid())
+	{
+		FArchive* FileReader = Attachment->CreateReader();
+
+		if (FileReader != nullptr)
+		{
+			FString DeploymentFolder = FPaths::EngineIntermediateDir() / TEXT("Deploy") / Message.TransactionId.ToString();
+			FString TargetPath = DeploymentFolder / Message.TargetFileName;
+
+			StoreDeployedFile(FileReader, TargetPath);
+
+			delete FileReader;
+		}
+	}
+}
+
+
+void FTargetDeviceService::HandleDeployCommitMessage(const FTargetDeviceServiceDeployCommit& Message, const IMessageContextRef& Context)
+{
+	if (!Running)
+	{
+		return;
+	}
+
+	ITargetDevicePtr TargetDevice = GetDevice(Message.Variant);
+
+	if (TargetDevice.IsValid())
+	{
+		FString SourceFolder = FPaths::EngineIntermediateDir() / TEXT("Deploy") / Message.TransactionId.ToString();
+		FString OutAppId;
+
+		bool Succeeded = TargetDevice->Deploy(SourceFolder, OutAppId);
+
+		IFileManager::Get().DeleteDirectory(*SourceFolder, false, true);
+		MessageEndpoint->Send(new FTargetDeviceServiceDeployFinished(Message.Variant, OutAppId, Succeeded, Message.TransactionId), Context->GetSender());
+	}
+}
+
+
+void FTargetDeviceService::HandleLaunchAppMessage(const FTargetDeviceServiceLaunchApp& Message, const IMessageContextRef& Context)
+{
+	if (!Running)
+	{
+		return;
+	}
+
+	ITargetDevicePtr TargetDevice = GetDevice(Message.Variant);
+
+	if (TargetDevice.IsValid())
+	{
+		uint32 ProcessId;
+		bool Succeeded = TargetDevice->Launch(Message.AppID, (EBuildConfigurations::Type)Message.BuildConfiguration, EBuildTargets::Game, Message.Params, &ProcessId);
+
+		if (MessageEndpoint.IsValid())
+		{
+			MessageEndpoint->Send(new FTargetDeviceServiceLaunchFinished(Message.AppID, ProcessId, Succeeded), Context->GetSender());
+		}
+	}
+}
+
+
 void FTargetDeviceService::HandlePingMessage(const FTargetDeviceServicePing& InMessage, const IMessageContextRef& Context)
 {
 	if (!Running)
@@ -80,53 +380,122 @@ void FTargetDeviceService::HandlePingMessage(const FTargetDeviceServicePing& InM
 		return;
 	}
 
-	if (Shared || (InMessage.HostUser == FPlatformProcess::UserName(true)))
+	if (!Shared && (InMessage.HostUser != FPlatformProcess::UserName(true)))
+	{
+		return;
+	}
+
+	ITargetDevicePtr DefaultDevice = GetDevice(); // Default Device is needed here!
+
+	if (DefaultDevice.IsValid())
 	{
-		ITargetDevicePtr DefaultDevice = GetDevice(); // Default Device is needed here!
+		const FString& PlatformName = DefaultDevice->GetTargetPlatform().PlatformName();
+		const PlatformInfo::FPlatformInfo* VanillaInfo = PlatformInfo::FindVanillaPlatformInfo(FName(*PlatformName));
+
+		FTargetDeviceServicePong* Message = new FTargetDeviceServicePong();
+
+		Message->Name = DefaultDevice->GetName();
+		Message->Type = ETargetDeviceTypes::ToString(DefaultDevice->GetDeviceType());
+		Message->HostName = FPlatformProcess::ComputerName();
+		Message->HostUser = FPlatformProcess::UserName(true);
+		Message->Connected = DefaultDevice->IsConnected();
+		Message->Make = TEXT("@todo");
+		Message->Model = TEXT("@todo");
+		DefaultDevice->GetUserCredentials(Message->DeviceUser, Message->DeviceUserPassword);
+		Message->Shared = Shared;
+		Message->SupportsMultiLaunch = DefaultDevice->SupportsFeature(ETargetDeviceFeatures::MultiLaunch);
+		Message->SupportsPowerOff = DefaultDevice->SupportsFeature(ETargetDeviceFeatures::PowerOff);
+		Message->SupportsPowerOn = DefaultDevice->SupportsFeature(ETargetDeviceFeatures::PowerOn);
+		Message->SupportsReboot = DefaultDevice->SupportsFeature(ETargetDeviceFeatures::Reboot);
+		Message->SupportsVariants = DefaultDevice->GetTargetPlatform().SupportsVariants();
+		Message->DefaultVariant = FName(DefaultDevice->GetTargetPlatform().PlatformName().GetCharArray().GetData());
+
+		// Add the data for all the flavors
+		Message->Variants.SetNumZeroed(TargetDevicePtrs.Num());
 
-		if (DefaultDevice.IsValid())
+		int Index = 0;
+		for (auto TargetDeviceIt = TargetDevicePtrs.CreateIterator(); TargetDeviceIt; ++TargetDeviceIt, ++Index)
 		{
-			const FString& PlatformName = DefaultDevice->GetTargetPlatform().PlatformName();
-			const PlatformInfo::FPlatformInfo* VanillaInfo = PlatformInfo::FindVanillaPlatformInfo(FName(*PlatformName));
-
-			FTargetDeviceServicePong* Message = new FTargetDeviceServicePong();
-
-			Message->Name = DefaultDevice->GetName();
-			Message->Type = ETargetDeviceTypes::ToString(DefaultDevice->GetDeviceType());
-			Message->HostName = FPlatformProcess::ComputerName();
-			Message->HostUser = FPlatformProcess::UserName(true);
-			Message->Connected = DefaultDevice->IsConnected();
-			Message->Make = TEXT("@todo");
-			Message->Model = TEXT("@todo");
-			DefaultDevice->GetUserCredentials(Message->DeviceUser, Message->DeviceUserPassword);
-			Message->Shared = Shared;
-			Message->SupportsMultiLaunch = DefaultDevice->SupportsFeature(ETargetDeviceFeatures::MultiLaunch);
-			Message->SupportsPowerOff = DefaultDevice->SupportsFeature(ETargetDeviceFeatures::PowerOff);
-			Message->SupportsPowerOn = DefaultDevice->SupportsFeature(ETargetDeviceFeatures::PowerOn);
-			Message->SupportsReboot = DefaultDevice->SupportsFeature(ETargetDeviceFeatures::Reboot);
-			Message->SupportsVariants = DefaultDevice->GetTargetPlatform().SupportsVariants();
-			Message->DefaultVariant = FName(DefaultDevice->GetTargetPlatform().PlatformName().GetCharArray().GetData());
-
-			// Add the data for all the flavors
-			Message->Variants.SetNumZeroed(TargetDevicePtrs.Num());
-
-			int Index = 0;
-			for (auto TargetDeviceIt = TargetDevicePtrs.CreateIterator(); TargetDeviceIt; ++TargetDeviceIt, ++Index)
-			{
-				const ITargetDevicePtr& TargetDevice = TargetDeviceIt.Value().Pin();
-				const PlatformInfo::FPlatformInfo& Info = TargetDevice->GetTargetPlatform().GetPlatformInfo();
-
-				FTargetDeviceVariant& Variant = Message->Variants[Index];
-
-				Variant.DeviceID = TargetDevice->GetId().ToString();
-				Variant.VariantName = TargetDeviceIt.Key();
-				Variant.TargetPlatformName = TargetDevice->GetTargetPlatform().PlatformName();
-				Variant.TargetPlatformId = Info.TargetPlatformName;
-				Variant.VanillaPlatformId = Info.VanillaPlatformName;
-				Variant.PlatformDisplayName = Info.DisplayName.ToString();
-			}
-
-			MessageEndpoint->Send(Message, Context->GetSender());
+			const ITargetDevicePtr& TargetDevice = TargetDeviceIt.Value().Pin();
+			const PlatformInfo::FPlatformInfo& Info = TargetDevice->GetTargetPlatform().GetPlatformInfo();
+
+			FTargetDeviceVariant& Variant = Message->Variants[Index];
+
+			Variant.DeviceID = TargetDevice->GetId().ToString();
+			Variant.VariantName = TargetDeviceIt.Key();
+			Variant.TargetPlatformName = TargetDevice->GetTargetPlatform().PlatformName();
+			Variant.TargetPlatformId = Info.TargetPlatformName;
+			Variant.VanillaPlatformId = Info.VanillaPlatformName;
+			Variant.PlatformDisplayName = Info.DisplayName.ToString();
 		}
+
+		MessageEndpoint->Send(Message, Context->GetSender());
+	}
+}
+
+
+void FTargetDeviceService::HandlePowerOffMessage( const FTargetDeviceServicePowerOff& Message, const IMessageContextRef& Context )
+{
+	if (!Running)
+	{
+		return;
 	}
-}
\ No newline at end of file
+
+	ITargetDevicePtr TargetDevice = GetDevice(); // Any Device is fine here
+
+	if (TargetDevice.IsValid())
+	{
+		TargetDevice->PowerOff(Message.Force);
+	}
+}
+
+
+void FTargetDeviceService::HandlePowerOnMessage( const FTargetDeviceServicePowerOn& Message, const IMessageContextRef& Context )
+{
+	if (!Running)
+	{
+		return;
+	}
+
+	ITargetDevicePtr TargetDevice = GetDevice(); // Any Device is fine here
+
+	if (TargetDevice.IsValid())
+	{
+		TargetDevice->PowerOn();
+	}
+}
+
+
+void FTargetDeviceService::HandleRebootMessage( const FTargetDeviceServiceReboot& Message, const IMessageContextRef& Context )
+{
+	if (!Running)
+	{
+		return;
+	}
+
+	ITargetDevicePtr TargetDevice = GetDevice(); // Any Device is fine here
+
+	if (TargetDevice.IsValid())
+	{
+		TargetDevice->Reboot();
+	}
+}
+
+
+void FTargetDeviceService::HandleRunExecutableMessage(const FTargetDeviceServiceRunExecutable& Message, const IMessageContextRef& Context)
+{
+	if (!Running)
+	{
+		return;
+	}
+
+	ITargetDevicePtr TargetDevice = GetDevice(Message.Variant);
+
+	if (TargetDevice.IsValid())
+	{
+		uint32 OutProcessId;
+		bool Succeeded = TargetDevice->Run(Message.ExecutablePath, Message.Params, &OutProcessId);
+
+		MessageEndpoint->Send(new FTargetDeviceServiceRunFinished(Message.Variant, Message.ExecutablePath, OutProcessId, Succeeded), Context->GetSender());
+	}
+}
diff --git a/Engine/Source/Developer/TargetDeviceServices/Private/TargetDeviceService.h b/Engine/Source/Developer/TargetDeviceServices/Private/TargetDeviceService.h
--- a/Engine/Source/Developer/TargetDeviceServices/Private/TargetDeviceService.h
+++ b/Engine/Source/Developer/TargetDeviceServices/Private/TargetDeviceService.h
@@ -8,16 +8,6 @@
 class FTargetDeviceService
 	: public ITargetDeviceService
 {
-	struct FVariantSortCallback
-	{
-		FORCEINLINE bool operator()(const ITargetDeviceWeakPtr& A, const ITargetDeviceWeakPtr& B) const
-		{
-			ITargetDevicePtr APtr = A.Pin();
-			ITargetDevicePtr BPtr = B.Pin();
-			return APtr->GetTargetPlatform().GetVariantPriority() > BPtr->GetTargetPlatform().GetVariantPriority();
-		}
-	};
-
 public:
 
 	/**
@@ -26,66 +16,17 @@ public:
 	 * @param InDeviceName The name of the device to expose.
 	 * @param InMessageBus The message bus to listen on for clients.
 	 */
-	FTargetDeviceService(const FString& InDeviceName, const IMessageBusRef& InMessageBus)
-		: DeviceName(InDeviceName)
-		, Running(false)
-		, Shared(false)
-	{
-		// initialize messaging
-		MessageEndpoint = FMessageEndpoint::Builder(FName(*FString::Printf(TEXT("FTargetDeviceService (%s)"), *DeviceName)), InMessageBus)
-			.Handling<FTargetDeviceClaimDenied>(this, &FTargetDeviceService::HandleClaimDeniedMessage)
-			.Handling<FTargetDeviceClaimed>(this, &FTargetDeviceService::HandleClaimedMessage)
-			.Handling<FTargetDeviceServiceDeployCommit>(this, &FTargetDeviceService::HandleDeployCommitMessage)
-			.Handling<FTargetDeviceServiceDeployFile>(this, &FTargetDeviceService::HandleDeployFileMessage)
-			.Handling<FTargetDeviceServiceLaunchApp>(this, &FTargetDeviceService::HandleLaunchAppMessage)
-			.Handling<FTargetDeviceServicePing>(this, &FTargetDeviceService::HandlePingMessage)
-			.Handling<FTargetDeviceServicePowerOff>(this, &FTargetDeviceService::HandlePowerOffMessage)
-			.Handling<FTargetDeviceServicePowerOn>(this, &FTargetDeviceService::HandlePowerOnMessage)
-			.Handling<FTargetDeviceServiceReboot>(this, &FTargetDeviceService::HandleRebootMessage)
-			.Handling<FTargetDeviceServiceRunExecutable>(this, &FTargetDeviceService::HandleRunExecutableMessage)
-			.Handling<FTargetDeviceUnclaimed>(this, &FTargetDeviceService::HandleUnclaimedMessage);
-
-		if (MessageEndpoint.IsValid())
-		{
-			MessageEndpoint->Subscribe<FTargetDeviceClaimed>();
-			MessageEndpoint->Subscribe<FTargetDeviceUnclaimed>();
-			MessageEndpoint->Subscribe<FTargetDeviceServicePing>();
-		}
-	}
+	FTargetDeviceService(const FString& InDeviceName, const IMessageBusRef& InMessageBus);
 
-	/**
-	 * Destructor.
-	 */
-	~FTargetDeviceService( )
-	{
-		Stop();
-		FMessageEndpoint::SafeRelease(MessageEndpoint);
-	}
+	/** Destructor. */
+	~FTargetDeviceService();
 
 public:
 
 	// ITargetDeviceService interface
 
 	virtual void AddTargetDevice(ITargetDevicePtr InDevice) override;
-
-	virtual void RemoveTargetDevice(ITargetDevicePtr InDevice) override;
-
-	virtual int32 NumTargetDevices() override
-	{
-		return TargetDevicePtrs.Num();
-	}
-
-	virtual bool CanStart(FName InFlavor = NAME_None) const override
-	{
-		ITargetDevicePtr TargetDevice = GetDevice(InFlavor);
-
-		if (TargetDevice.IsValid())
-		{
-			return TargetDevice->IsConnected();
-		}
-
-		return false;
-	}
+	virtual bool CanStart(FName InFlavor = NAME_None) const override;
 
 	virtual const FString& GetClaimHost() override
 	{
@@ -97,23 +38,7 @@ public:
 		return ClaimUser;
 	}
 
-	virtual ITargetDevicePtr GetDevice(FName InVariant = NAME_None) const override
-	{
-		ITargetDevicePtr TargetDevice;
-		if (InVariant == NAME_None)
-		{
-			TargetDevice = DefaultDevicePtr.Pin();
-		}
-		else
-		{
-			const ITargetDeviceWeakPtr * WeakTargetDevicePtr = TargetDevicePtrs.Find(InVariant);
-			if (WeakTargetDevicePtr != NULL)
-			{
-				TargetDevice = WeakTargetDevicePtr->Pin();
-			}
-		}
-		return TargetDevice;
-	}
+	virtual ITargetDevicePtr GetDevice(FName InVariant = NAME_None) const override;
 
 	virtual FString GetDeviceName() const override
 	{
@@ -140,312 +65,99 @@ public:
 		return (Running && Shared);
 	}
 
-	virtual void SetShared(bool InShared) override
+	virtual int32 NumTargetDevices() override
 	{
-		Shared = InShared;
+		return TargetDevicePtrs.Num();
 	}
 
-	virtual bool Start() override
-	{
-		if (!Running && MessageEndpoint.IsValid())
-		{
-			// notify other services
-			ClaimAddress = MessageEndpoint->GetAddress();
-			ClaimHost = FPlatformProcess::ComputerName();
-			ClaimUser = FPlatformProcess::UserName(true);
-
-			MessageEndpoint->Publish(new FTargetDeviceClaimed(DeviceName, ClaimHost, ClaimUser));
 
-			Running = true;
-		}
-
-		return true;
-	}
+	virtual void RemoveTargetDevice(ITargetDevicePtr InDevice) override;
 
-	virtual void Stop() override
+	virtual void SetShared(bool InShared) override
 	{
-		if (Running)
-		{
-			// notify other services
-			MessageEndpoint->Publish(new FTargetDeviceUnclaimed(DeviceName, FPlatformProcess::ComputerName(), FPlatformProcess::UserName(true)));
-
-			Running = false;
-		}
+		Shared = InShared;
 	}
 
+	virtual bool Start() override;
+	virtual void Stop() override;
+
 protected:
 
 	/**
-	* Stores the specified file to deploy.
-	*
-	* @param FileReader The archive reader providing access to the file data.
-	* @param TargetFileName The desired name of the file on the target device.
-	*/
-	bool StoreDeployedFile(FArchive* FileReader, const FString& TargetFileName) const
-	{
-		if (FileReader == NULL)
-		{
-			return false;
-		}
-
-		// create target file
-		FArchive* FileWriter = IFileManager::Get().CreateFileWriter(*TargetFileName);
-
-		if (FileWriter == NULL)
-		{
-			return false;
-		}
-
-		FileReader->Seek(0);
-
-		// copy file contents
-		int64 BytesRemaining = FileReader->TotalSize();
-		int32 BufferSize = 128 * 1024;
-
-		if (BytesRemaining < BufferSize)
-		{
-			BufferSize = BytesRemaining;
-		}
-
-		void* Buffer = FMemory::Malloc(BufferSize);
-
-		while (BytesRemaining > 0)
-		{
-			FileReader->Serialize(Buffer, BufferSize);
-			FileWriter->Serialize(Buffer, BufferSize);
-
-			BytesRemaining -= BufferSize;
-
-			if (BytesRemaining < BufferSize)
-			{
-				BufferSize = BytesRemaining;
-			}
-		}
-
-		// clean up
-		FMemory::Free(Buffer);
-
-		delete FileWriter;
-
-		return true;
-	}
+	 * Stores the specified file to deploy.
+	 *
+	 * @param FileReader The archive reader providing access to the file data.
+	 * @param TargetFileName The desired name of the file on the target device.
+	 */
+	bool StoreDeployedFile(FArchive* FileReader, const FString& TargetFileName) const;
 
 private:
 
-	// Callback for FTargetDeviceClaimDenied messages.
-	void HandleClaimDeniedMessage( const FTargetDeviceClaimDenied& Message, const IMessageContextRef& Context )
-	{
-		if (Running && (Message.DeviceName == DeviceName))
-		{
-			Stop();
-
-			ClaimAddress = Context->GetSender();
-			ClaimHost = Message.HostName;
-			ClaimUser = Message.HostUser;
-		}
-	}
-
-	// Callback for FTargetDeviceClaimDenied messages.
-	void HandleClaimedMessage( const FTargetDeviceClaimed& Message, const IMessageContextRef& Context )
-	{
-		if (Message.DeviceName == DeviceName)
-		{
-			if (Running)
-			{
-				if (Context->GetSender() != MessageEndpoint->GetAddress())
-				{
-					MessageEndpoint->Send(new FTargetDeviceClaimDenied(DeviceName, FPlatformProcess::ComputerName(), FPlatformProcess::UserName(true)), Context->GetSender());
-				}
-			}
-			else
-			{
-				ClaimAddress = Context->GetSender();
-				ClaimHost = Message.HostName;
-				ClaimUser = Message.HostUser;
-			}
-		}
-	}
-
-	// Callback for FTargetDeviceClaimDropped messages.
-	void HandleUnclaimedMessage( const FTargetDeviceUnclaimed& Message, const IMessageContextRef& Context )
-	{
-		if (Message.DeviceName == DeviceName)
-		{
-			if (Context->GetSender() == ClaimAddress)
-			{
-				ClaimAddress = FGuid();
-				ClaimHost.Empty();
-				ClaimUser.Empty();
-			}
-		}
-	}
-
-	// Callback for FTargetDeviceServiceDeployFile messages.
-	void HandleDeployFileMessage(const FTargetDeviceServiceDeployFile& Message, const IMessageContextRef& Context)
-	{
-		if (!Running)
-		{
-			return;
-		}
-
-		IMessageAttachmentPtr Attachment = Context->GetAttachment();
+	/** Callback for FTargetDeviceClaimDenied messages. */
+	void HandleClaimDeniedMessage(const FTargetDeviceClaimDenied& Message, const IMessageContextRef& Context);
 
-		if (Attachment.IsValid())
-		{
-			FArchive* FileReader = Attachment->CreateReader();
+	/** Callback for FTargetDeviceClaimDenied messages. */
+	void HandleClaimedMessage( const FTargetDeviceClaimed& Message, const IMessageContextRef& Context );
 
-			if (FileReader != NULL)
-			{
-				FString DeploymentFolder = FPaths::EngineIntermediateDir() / TEXT("Deploy") / Message.TransactionId.ToString();
-				FString TargetPath = DeploymentFolder / Message.TargetFileName;
+	/** Callback for FTargetDeviceClaimDropped messages. */
+	void HandleUnclaimedMessage(const FTargetDeviceUnclaimed& Message, const IMessageContextRef& Context);
 
-				StoreDeployedFile(FileReader, TargetPath);
+	/** Callback for FTargetDeviceServiceDeployFile messages. */
+	void HandleDeployFileMessage(const FTargetDeviceServiceDeployFile& Message, const IMessageContextRef& Context);
 
-				delete FileReader;
-			}
-		}
-	}
-
-	// Callback for FTargetDeviceServiceDeployCommit messages.
-	void HandleDeployCommitMessage(const FTargetDeviceServiceDeployCommit& Message, const IMessageContextRef& Context)
-	{
-		if (!Running)
-		{
-			return;
-		}
-
-		ITargetDevicePtr TargetDevice = GetDevice(Message.Variant);
-		if (TargetDevice.IsValid())
-		{
-			FString SourceFolder = FPaths::EngineIntermediateDir() / TEXT("Deploy") / Message.TransactionId.ToString();
-			FString OutAppId;
-
-			bool Succeeded = TargetDevice->Deploy(SourceFolder, OutAppId);
-
-			IFileManager::Get().DeleteDirectory(*SourceFolder, false, true);
-			MessageEndpoint->Send(new FTargetDeviceServiceDeployFinished(Message.Variant, OutAppId, Succeeded, Message.TransactionId), Context->GetSender());
-		}
-	}
+	/** Callback for FTargetDeviceServiceDeployCommit messages. */
+	void HandleDeployCommitMessage(const FTargetDeviceServiceDeployCommit& Message, const IMessageContextRef& Context);
 
-	// Callback for FTargetDeviceServiceLaunchApp messages.
-	void HandleLaunchAppMessage(const FTargetDeviceServiceLaunchApp& Message, const IMessageContextRef& Context)
-	{
-		if (!Running)
-		{
-			return;
-		}
-
-		ITargetDevicePtr TargetDevice = GetDevice(Message.Variant);
-		if (TargetDevice.IsValid())
-		{
-			uint32 ProcessId;
-			bool Succeeded = TargetDevice->Launch(Message.AppID, (EBuildConfigurations::Type)Message.BuildConfiguration, EBuildTargets::Game, Message.Params, &ProcessId);
-
-			if (MessageEndpoint.IsValid())
-			{
-				MessageEndpoint->Send(new FTargetDeviceServiceLaunchFinished(Message.AppID, ProcessId, Succeeded), Context->GetSender());
-			}
-		}
-	}
+	/** Callback for FTargetDeviceServiceLaunchApp messages. */
+	void HandleLaunchAppMessage(const FTargetDeviceServiceLaunchApp& Message, const IMessageContextRef& Context);
 
-	// Callback for FTargetDeviceServicePing messages.
+	/** Callback for FTargetDeviceServicePing messages. */
 	void HandlePingMessage(const FTargetDeviceServicePing& InMessage, const IMessageContextRef& Context);
 
-	// Callback for FTargetDeviceServicePowerOff messages.
-	void HandlePowerOffMessage( const FTargetDeviceServicePowerOff& Message, const IMessageContextRef& Context )
-	{
-		if (!Running)
-		{
-			return;
-		}
-
-		ITargetDevicePtr TargetDevice = GetDevice(); // Any Device is fine here
-		if (TargetDevice.IsValid())
-		{
-			TargetDevice->PowerOff(Message.Force);
-		}
-	}
+	 /** Callback for FTargetDeviceServicePowerOff messages. */
+	void HandlePowerOffMessage(const FTargetDeviceServicePowerOff& Message, const IMessageContextRef& Context);
 
-	// Callback for FTargetDeviceServicePowerOn messages.
-	void HandlePowerOnMessage( const FTargetDeviceServicePowerOn& Message, const IMessageContextRef& Context )
-	{
-		if (!Running)
-		{
-			return;
-		}
-
-		ITargetDevicePtr TargetDevice = GetDevice(); // Any Device is fine here
-		if (TargetDevice.IsValid())
-		{
-			TargetDevice->PowerOn();
-		}
-	}
+	/** Callback for FTargetDeviceServicePowerOn messages. */
+	void HandlePowerOnMessage(const FTargetDeviceServicePowerOn& Message, const IMessageContextRef& Context);
 
-	// Callback for FTargetDeviceServiceReboot messages.
-	void HandleRebootMessage( const FTargetDeviceServiceReboot& Message, const IMessageContextRef& Context )
-	{
-		if (!Running)
-		{
-			return;
-		}
-
-		ITargetDevicePtr TargetDevice = GetDevice(); // Any Device is fine here
-		if (TargetDevice.IsValid())
-		{
-			TargetDevice->Reboot();
-		}
-	}
+	/** Callback for FTargetDeviceServiceReboot messages. */
+	void HandleRebootMessage(const FTargetDeviceServiceReboot& Message, const IMessageContextRef& Context);
 
-	// Callback for FTargetDeviceServiceRunExecutable messages.
-	void HandleRunExecutableMessage(const FTargetDeviceServiceRunExecutable& Message, const IMessageContextRef& Context)
-	{
-		if (!Running)
-		{
-			return;
-		}
-
-		ITargetDevicePtr TargetDevice = GetDevice(Message.Variant);
-		if (TargetDevice.IsValid())
-		{
-			uint32 OutProcessId;
-			bool Succeeded = TargetDevice->Run(Message.ExecutablePath, Message.Params, &OutProcessId);
-
-			MessageEndpoint->Send(new FTargetDeviceServiceRunFinished(Message.Variant, Message.ExecutablePath, OutProcessId, Succeeded), Context->GetSender());
-		}
-	}
+	/** Callback for FTargetDeviceServiceRunExecutable messages. */
+	void HandleRunExecutableMessage(const FTargetDeviceServiceRunExecutable& Message, const IMessageContextRef& Context);
 
 private:
 
-	// Caches the name of the device name that this services exposes.
+	/** Default target device used when no flavor is specified. */
+	ITargetDeviceWeakPtr DefaultDevicePtr;
+
+	/** Caches the name of the device name that this services exposes. */
 	FString DeviceName;
 
-	// Caches the platform name of the device name that this services exposes.
+	/** Caches the platform name of the device name that this services exposes. */
 	FName DevicePlatformName;
 
-	// Caches the platform name of the device name that this services exposes.
+	/** Caches the platform name of the device name that this services exposes. */
 	FString DevicePlatformDisplayName;
 
-	// Holds the name of the host that has a claim on the device.
+	/** Holds the name of the host that has a claim on the device. */
 	FString ClaimHost;
 
-	// Holds the message address of the target device service that has a claim on the device.
+	/** Holds the message address of the target device service that has a claim on the device. */
 	FMessageAddress ClaimAddress;
 
-	// Holds the name of the user that has a claim on the device.
+	/** Holds the name of the user that has a claim on the device. */
 	FString ClaimUser;
 
-	// Holds the message endpoint.
+	/** Holds the message endpoint. */
 	FMessageEndpointPtr MessageEndpoint;
 
-	// Holds a flag indicating whether this service is currently running.
+	/** Holds a flag indicating whether this service is currently running. */
 	bool Running;
 
-	// Holds a flag indicating whether the device is shared with other users.
+	/** Holds a flag indicating whether the device is shared with other users. */
 	bool Shared;
-
-	// Default target device used when no flavor is specified
-	ITargetDeviceWeakPtr DefaultDevicePtr;
 	
-	// Map of all the Flavors for this Service
+	/** Map of all the Flavors for this Service */
 	TMap<FName, ITargetDeviceWeakPtr> TargetDevicePtrs;
 };
diff --git a/Engine/Source/Developer/TargetDeviceServices/Private/TargetDeviceServiceManager.h b/Engine/Source/Developer/TargetDeviceServices/Private/TargetDeviceServiceManager.h
--- a/Engine/Source/Developer/TargetDeviceServices/Private/TargetDeviceServiceManager.h
+++ b/Engine/Source/Developer/TargetDeviceServices/Private/TargetDeviceServiceManager.h
@@ -13,15 +13,11 @@ class FTargetDeviceServiceManager
 {
 public:
 
-	/**
-	 * Default constructor.
-	 */
-	FTargetDeviceServiceManager( );
+	/** Default constructor. */
+	FTargetDeviceServiceManager();
 
-	/**
-	 * Destructor.
-	 */
-	~FTargetDeviceServiceManager( );
+	/** Destructor. */
+	~FTargetDeviceServiceManager();
 
 public:
 
@@ -55,7 +51,7 @@ protected:
 	ITargetDeviceServicePtr AddService(const FString& DeviceName);
 
 	/**
-	 * Adds a target device, a device service will be created if needed
+	 * Adds a target device, a device service will be created if needed.
 	 *
 	 * @param InDevice The target device to add
 	 * @return true if the device was added to the device service, false if the service failed to create.
@@ -106,34 +102,34 @@ protected:
 
 private:
 
-	// Callback for handling message bus shutdowns.
+	/** Callback for handling message bus shutdowns. */
 	void HandleMessageBusShutdown();
 
-	// Callback for discovered target devices.
+	/** Callback for discovered target devices. */
 	void HandleTargetPlatformDeviceDiscovered(ITargetDeviceRef DiscoveredDevice);
 
-	// Callback for lost target devices.
+	/** Callback for lost target devices. */
 	void HandleTargetPlatformDeviceLost(ITargetDeviceRef LostDevice);
 
 private:
 
-	// Holds a critical section object.
+	/** Holds a critical section object. */
 	FCriticalSection CriticalSection;
 
-	// Holds the list of locally managed device services.
+	/** Holds the list of locally managed device services. */
 	TMap<FString, ITargetDeviceServicePtr> DeviceServices;
 
-	// Holds a weak pointer to the message bus.
+	/** Holds a weak pointer to the message bus. */
 	IMessageBusWeakPtr MessageBusPtr;
 
-	// Holds the collection of identifiers for devices that start automatically (shared/unshared).
+	/** Holds the collection of identifiers for devices that start automatically (shared/unshared). */
 	TMap<FString, bool> StartupServices;
 
 private:
 
-	// Holds a delegate that is invoked when a target device service was added.
+	/** Holds a delegate that is invoked when a target device service was added. */
 	FOnTargetDeviceServiceAdded ServiceAddedDelegate;
 
-	// Holds a delegate that is invoked when a target device service was removed.
+	/** Holds a delegate that is invoked when a target device service was removed. */
 	FOnTargetDeviceServiceRemoved ServiceRemovedDelegate;
 };
diff --git a/Engine/Source/Developer/TargetDeviceServices/Private/TargetDeviceServicesModule.cpp b/Engine/Source/Developer/TargetDeviceServices/Private/TargetDeviceServicesModule.cpp
--- a/Engine/Source/Developer/TargetDeviceServices/Private/TargetDeviceServicesModule.cpp
+++ b/Engine/Source/Developer/TargetDeviceServices/Private/TargetDeviceServicesModule.cpp
@@ -16,7 +16,7 @@ public:
 
 	// ITargetDeviceServicesModule interface
 
-	virtual ITargetDeviceProxyManagerRef GetDeviceProxyManager( ) override
+	virtual ITargetDeviceProxyManagerRef GetDeviceProxyManager() override
 	{
 		if (!DeviceProxyManagerSingleton.IsValid())
 		{
@@ -26,7 +26,7 @@ public:
 		return DeviceProxyManagerSingleton.ToSharedRef();
 	}
 
-	virtual ITargetDeviceServiceManagerRef GetDeviceServiceManager( ) override
+	virtual ITargetDeviceServiceManagerRef GetDeviceServiceManager() override
 	{
 		if (!DeviceServiceManagerSingleton.IsValid())
 		{
@@ -38,10 +38,10 @@ public:
 
 private:
 
-	// Holds the device proxy manager singleton.
+	/** Holds the device proxy manager singleton. */
 	ITargetDeviceProxyManagerPtr DeviceProxyManagerSingleton;
 
-	// Holds the device service manager singleton.
+	/** Holds the device service manager singleton. */
 	ITargetDeviceServiceManagerPtr DeviceServiceManagerSingleton;
 };
 
diff --git a/Engine/Source/Developer/TargetDeviceServices/Public/Interfaces/ITargetDeviceProxy.h b/Engine/Source/Developer/TargetDeviceServices/Public/Interfaces/ITargetDeviceProxy.h
--- a/Engine/Source/Developer/TargetDeviceServices/Public/Interfaces/ITargetDeviceProxy.h
+++ b/Engine/Source/Developer/TargetDeviceServices/Public/Interfaces/ITargetDeviceProxy.h
@@ -40,6 +40,7 @@ DECLARE_MULTICAST_DELEGATE_OneParam(FOnTargetDeviceProxyLaunchFailed, const FStr
  */
 DECLARE_MULTICAST_DELEGATE_TwoParams(FOnTargetDeviceProxyLaunchSucceeded, const FString&, uint32)
 
+
 /**
  * Interface for target device proxies.
  */
@@ -252,9 +253,7 @@ public:
 	 * @param Files The files to deploy.
 	 * @param TransactionId The transaction identifier for the deployment.
 	 * @return true if deployment has been started, false otherwise.
-	 * @see Launch
-	 * @see OnDeployCommitted
-	 * @see OnDeployFailed
+	 * @see Launch, OnDeployCommitted, OnDeployFailed
 	 */
 	virtual bool DeployApp(FName InVariant, const TMap<FString, FString>& Files, const FGuid& TransactionId) = 0;
 
@@ -268,8 +267,7 @@ public:
 	 * @param BuildConfiguration The build configuration to launch with.
 	 * @param Params The command line parameters to launch with.
 	 * @return true on success, false otherwise.
-	 * @see Deploy
-	 * @see OnLaunchFailed
+	 * @see Deploy, OnLaunchFailed
 	 */
 	virtual bool LaunchApp(FName InVariant, const FString& AppId, EBuildConfigurations::Type BuildConfiguration, const FString& Params) = 0;
 
@@ -277,24 +275,21 @@ public:
 	 * Powers off the device.
 	 *
 	 * @param Force Whether to force powering off.
-	 * @see PowerOn
-	 * @see Reboot
+	 * @see PowerOn, Reboot
 	 */
 	virtual void PowerOff(bool Force) = 0;
 
 	/**
 	 * Powers on the device.
 	 *
-	 * @see PowerOff
-	 * @see Reboot
+	 * @see PowerOff, Reboot
 	 */
 	virtual void PowerOn() = 0;
 
 	/**
 	 * Reboots the device.
 	 *
-	 * @see PowerOff
-	 * @see PowerOn
+	 * @see PowerOff, PowerOn
 	 */
 	virtual void Reboot() = 0;
 
@@ -335,8 +330,6 @@ public:
 
 public:
 
-	/**
-	 * Virtual destructor.
-	 */
+	/** Virtual destructor. */
 	virtual ~ITargetDeviceProxy() { }
 };
diff --git a/Engine/Source/Developer/TargetDeviceServices/Public/Interfaces/ITargetDeviceProxyManager.h b/Engine/Source/Developer/TargetDeviceServices/Public/Interfaces/ITargetDeviceProxyManager.h
--- a/Engine/Source/Developer/TargetDeviceServices/Public/Interfaces/ITargetDeviceProxyManager.h
+++ b/Engine/Source/Developer/TargetDeviceServices/Public/Interfaces/ITargetDeviceProxyManager.h
@@ -9,6 +9,7 @@ typedef TSharedPtr<class ITargetDeviceProxyManager> ITargetDeviceProxyManagerPtr
 /** Type definition for shared references to instances of ITargetDeviceProxyManager. */
 typedef TSharedRef<class ITargetDeviceProxyManager> ITargetDeviceProxyManagerRef;
 
+
 /**
  * Interface for device proxy managers.
  */
@@ -21,6 +22,7 @@ public:
 	 *
 	 * @param Name The name of the device to create the proxy for.
 	 * @return The device proxy.
+	 * @see FindProxy, FindProxyDeviceForTargetDevice, GetProxies
 	 */
 	virtual ITargetDeviceProxyRef FindOrAddProxy(const FString& Name) = 0;
 
@@ -29,6 +31,7 @@ public:
 	 *
 	 * @param Name The name of the device to create the proxy for.
 	 * @return The device proxy, or nullptr if it couldn't be found.
+	 * @see FindOrAddProxy, FindProxyDeviceForTargetDevice, GetProxies
 	 */
 	virtual ITargetDeviceProxyPtr FindProxy(const FString& Name) = 0;
 
@@ -37,6 +40,7 @@ public:
 	 *
 	 * @param DeviceId The identifier of the target device to create the proxy for.
 	 * @return The device proxy, or nullptr if it couldn't be found.
+	 * @see FindOrAddProxy, FindProxy, GetProxies
 	 */
 	virtual ITargetDeviceProxyPtr FindProxyDeviceForTargetDevice(const FString& DeviceId) = 0;
 
@@ -46,6 +50,7 @@ public:
 	 * @param PlatformName The the name of the target platform to get proxies for (or empty string for all proxies).
 	 * @param IncludeUnshared Whether to include devices that are not being shared with the local user.
 	 * @param OutProxies Will hold the list of devices found by the locator
+	 * @see FindOrAddProxy, FindProxy, FindProxyDeviceForTargetDevice
 	 */
 	virtual void GetProxies(FName TargetPlatformName, bool IncludeUnshared, TArray<ITargetDeviceProxyPtr>& OutProxies) = 0;
 
@@ -57,7 +62,7 @@ public:
 	 * @return The event delegate.
 	 */
 	DECLARE_EVENT_OneParam(ITargetDeviceServiceManager, FOnTargetDeviceProxyAdded, const ITargetDeviceProxyRef& /*AddedProxy*/);
-	virtual FOnTargetDeviceProxyAdded& OnProxyAdded( ) = 0;
+	virtual FOnTargetDeviceProxyAdded& OnProxyAdded() = 0;
 
 	/**
 	 * Gets an event delegate that is executed when a target device proxy was removed.
@@ -65,9 +70,10 @@ public:
 	 * @return The event delegate.
 	 */
 	DECLARE_EVENT_OneParam(ITargetDeviceServiceManager, FOnTargetDeviceProxyRemoved, const ITargetDeviceProxyRef& /*RemovedProxy*/);
-	virtual FOnTargetDeviceProxyRemoved& OnProxyRemoved( ) = 0;
+	virtual FOnTargetDeviceProxyRemoved& OnProxyRemoved() = 0;
 
-protected:
+public:
 
-	ITargetDeviceProxyManager( ) { }
+	/** Virtual constructor. */
+	~ITargetDeviceProxyManager() { }
 };
diff --git a/Engine/Source/Developer/TargetDeviceServices/Public/Interfaces/ITargetDeviceService.h b/Engine/Source/Developer/TargetDeviceServices/Public/Interfaces/ITargetDeviceService.h
--- a/Engine/Source/Developer/TargetDeviceServices/Public/Interfaces/ITargetDeviceService.h
+++ b/Engine/Source/Developer/TargetDeviceServices/Public/Interfaces/ITargetDeviceService.h
@@ -30,20 +30,6 @@ public:
 	 */
 	virtual void AddTargetDevice(ITargetDevicePtr InDevice) = 0;
 
-	/**
-	 * Removes a flavor from this device.
-	 *
-	 * @return true if the service can start, false otherwise.
-	 */
-	virtual void RemoveTargetDevice(ITargetDevicePtr InDevice) = 0;
-
-	/**
-	 * Gets number of target devices.
-	 *
-	 * @return number of target devices in this device service.
-	 */
-	virtual int32 NumTargetDevices() = 0;
-
 	/**
 	 * Checks whether this service can start.
 	 *
@@ -59,14 +45,14 @@ public:
 	 *
 	 * @return Host name string.
 	 */
-	virtual const FString& GetClaimHost( ) = 0;
+	virtual const FString& GetClaimHost() = 0;
 
 	/**
 	 * Gets the name of the user that has a claim on the device.
 	 *
 	 * @return User name string.
 	 */
-	virtual const FString& GetClaimUser( ) = 0;
+	virtual const FString& GetClaimUser() = 0;
 
 	/**
 	 * Gets the target device for the supplied flavor, if no flavor is specified then the default device flavor is returned.
@@ -101,7 +87,7 @@ public:
 	 *
 	 * @return true if the service is running, false otherwise.
 	 */
-	virtual bool IsRunning( ) const = 0;
+	virtual bool IsRunning() const = 0;
 
 	/**
 	 * Checks whether the device is being shared with other users.
@@ -112,9 +98,21 @@ public:
 	 * @return true if the device is shared, false otherwise.
 	 * @see SetShared
 	 */
-	virtual bool IsShared( ) const = 0;
+	virtual bool IsShared() const = 0;
 
-public:
+	/**
+	 * Gets number of target devices.
+	 *
+	 * @return number of target devices in this device service.
+	 */
+	virtual int32 NumTargetDevices() = 0;
+
+	/**
+	 * Removes a flavor from this device.
+	 *
+	 * @return true if the service can start, false otherwise.
+	 */
+	virtual void RemoveTargetDevice(ITargetDevicePtr InDevice) = 0;
 
 	/**
 	 * Sets whether the device should be shared with other users.
@@ -125,29 +123,25 @@ public:
 	 * @param InShared Indicates whether device sharing is enabled.
 	 * @see IsShared
 	 */
-	virtual void SetShared( bool InShared ) = 0;
+	virtual void SetShared(bool InShared) = 0;
 
 	/**
 	 * Starts the service.
 	 *
 	 * @return true if the service has been started, false otherwise.
-	 * @see IsRunning
-	 * @see Stop
+	 * @see IsRunning, Stop
 	 */
-	virtual bool Start( ) = 0;
+	virtual bool Start() = 0;
 
 	/**
 	 * Stops the service.
 	 *
-	 * @see IsRunning
-	 * @see Stop
+	 * @see IsRunning, Stop
 	 */
-	virtual void Stop( ) = 0;
+	virtual void Stop() = 0;
 
 public:
 
-	/**
-	 * Virtual destructor.
-	 */
-	virtual ~ITargetDeviceService( ) { }
+	/** Virtual destructor. */
+	virtual ~ITargetDeviceService() { }
 };
diff --git a/Engine/Source/Developer/TargetDeviceServices/Public/Interfaces/ITargetDeviceServiceManager.h b/Engine/Source/Developer/TargetDeviceServices/Public/Interfaces/ITargetDeviceServiceManager.h
--- a/Engine/Source/Developer/TargetDeviceServices/Public/Interfaces/ITargetDeviceServiceManager.h
+++ b/Engine/Source/Developer/TargetDeviceServices/Public/Interfaces/ITargetDeviceServiceManager.h
@@ -26,8 +26,7 @@ public:
 	 *
 	 * @param DeviceName The device name for the service to add to the startup list
 	 * @return true if the service added, false otherwise.
-	 * @see GetServices
-	 * @see RemoveStartupService
+	 * @see GetServices, RemoveStartupService
 	 */
 	virtual bool AddStartupService(const FString& DeviceName) = 0;
 
@@ -37,14 +36,13 @@ public:
 	 * @param OutServices Will contain the collection of services.
 	 * @return The number of services returned.
 	 */
-	virtual int32 GetServices( TArray<ITargetDeviceServicePtr>& OutServices ) = 0;
+	virtual int32 GetServices(TArray<ITargetDeviceServicePtr>& OutServices) = 0;
 
 	/**
 	 * Removes a service from the list of services that are started automatically.
 	 *
 	 * @param DeviceName The device name for the service to remove from the startup list.
-	 * @see AddStartupService
-	 * @see GetServices
+	 * @see AddStartupService, GetServices
 	 */
 	virtual void RemoveStartupService(const FString& DeviceName) = 0;
 
@@ -56,7 +54,7 @@ public:
 	 * @return The event delegate.
 	 */
 	DECLARE_EVENT_OneParam(ITargetDeviceServiceManager, FOnTargetDeviceServiceAdded, const ITargetDeviceServiceRef& /*AddedService*/);
-	virtual FOnTargetDeviceServiceAdded& OnServiceAdded( ) = 0;
+	virtual FOnTargetDeviceServiceAdded& OnServiceAdded() = 0;
 
 	/**
 	 * Gets an event delegate that is executed when a target device service was removed.
@@ -64,12 +62,10 @@ public:
 	 * @return The event delegate.
 	 */
 	DECLARE_EVENT_OneParam(ITargetDeviceServiceManager, FOnTargetDeviceServiceRemoved, const ITargetDeviceServiceRef& /*RemovedService*/);
-	virtual FOnTargetDeviceServiceRemoved& OnServiceRemoved( ) = 0;
+	virtual FOnTargetDeviceServiceRemoved& OnServiceRemoved() = 0;
 
 public:
 
-	/**
-	 * Virtual destructor.
-	 */
-	virtual ~ITargetDeviceServiceManager( ) { }
+	/** Virtual destructor. */
+	virtual ~ITargetDeviceServiceManager() { }
 };
diff --git a/Engine/Source/Developer/TargetDeviceServices/Public/Interfaces/ITargetDeviceServicesModule.h b/Engine/Source/Developer/TargetDeviceServices/Public/Interfaces/ITargetDeviceServicesModule.h
--- a/Engine/Source/Developer/TargetDeviceServices/Public/Interfaces/ITargetDeviceServicesModule.h
+++ b/Engine/Source/Developer/TargetDeviceServices/Public/Interfaces/ITargetDeviceServicesModule.h
@@ -15,20 +15,20 @@ public:
 	 * Gets the target device proxy manager.
 	 *
 	 * @return The device proxy manager.
+	 * @see GetDeviceServiceManager
 	 */
-	virtual ITargetDeviceProxyManagerRef GetDeviceProxyManager( ) = 0;
+	virtual ITargetDeviceProxyManagerRef GetDeviceProxyManager() = 0;
 
 	/**
 	 * Gets the target device service manager.
 	 *
 	 * @return The device service manager.
+	 * @see GetDeviceProxyManager
 	 */
-	virtual ITargetDeviceServiceManagerRef GetDeviceServiceManager( ) = 0;
+	virtual ITargetDeviceServiceManagerRef GetDeviceServiceManager() = 0;
 
 public:
 
-	/**
-	 * Virtual destructor.
-	 */
-	virtual ~ITargetDeviceServicesModule( ) { }
+	/** Virtual destructor. */
+	virtual ~ITargetDeviceServicesModule() { }
 };
