Fix for intermitent add issue with async discovered devices.

diff --git a/Engine/Source/Developer/TargetDeviceServices/Private/TargetDeviceService.cpp b/Engine/Source/Developer/TargetDeviceServices/Private/TargetDeviceService.cpp
--- a/Engine/Source/Developer/TargetDeviceServices/Private/TargetDeviceService.cpp
+++ b/Engine/Source/Developer/TargetDeviceServices/Private/TargetDeviceService.cpp
@@ -87,23 +87,23 @@ void FTargetDeviceService::AddTargetDevice(ITargetDevicePtr InDevice)
 		DevicePlatformDisplayName = bSplit ? Left.Trim() : Full;
 	}
 
-	// make sure we don't have conflicting flavors for the same device!
-	// could also be a double add, which ideally would be avoided.
-	check(!(TargetDevicePtrs.FindRef(Variant).IsValid()));
-
-	TargetDevicePtrs.Add(Variant, InDevice);
+	// double add, which due to the async nature of some device discovery can't be easily avoided.
+	if (!(TargetDevicePtrs.FindRef(Variant).IsValid()))
+	{
+		TargetDevicePtrs.Add(Variant, InDevice);
 
-	// sort and choose cache the default
-	TargetDevicePtrs.ValueSort(FVariantSortCallback());
-	auto DeviceIterator = TargetDevicePtrs.CreateIterator();
+		// sort and choose cache the default
+		TargetDevicePtrs.ValueSort(FVariantSortCallback());
+		auto DeviceIterator = TargetDevicePtrs.CreateIterator();
 
-	if (DeviceIterator)
-	{
-		DefaultDevicePtr = (*DeviceIterator).Value;
-	}
-	else
-	{
-		DefaultDevicePtr = nullptr;
+		if (DeviceIterator)
+		{
+			DefaultDevicePtr = (*DeviceIterator).Value;
+		}
+		else
+		{
+			DefaultDevicePtr = nullptr;
+		}
 	}
 }
 
