Protect against race condition when adding targets

diff --git a/Engine/Source/Developer/TargetDeviceServices/Private/TargetDeviceServiceManager.cpp b/Engine/Source/Developer/TargetDeviceServices/Private/TargetDeviceServiceManager.cpp
--- a/Engine/Source/Developer/TargetDeviceServices/Private/TargetDeviceServiceManager.cpp
+++ b/Engine/Source/Developer/TargetDeviceServices/Private/TargetDeviceServiceManager.cpp
@@ -140,6 +140,9 @@ void FTargetDeviceServiceManager::InitializeTargetPlatforms()
 {
 	TArray<ITargetPlatform*> Platforms = GetTargetPlatformManager()->GetTargetPlatforms();
 
+	// make sure we're not handling any broadcast events while directly adding devices.
+	FScopeLock Lock(&CriticalSection);
+
 	for (int32 PlatformIndex = 0; PlatformIndex < Platforms.Num(); ++PlatformIndex)
 	{
 		// set up target platform callbacks
