Lock isn't necessary here since double-add is safe now.  Fixes a potential deadlock for platforms that have an internal critical section on device discovery.

diff --git a/Engine/Source/Developer/TargetDeviceServices/Private/TargetDeviceServiceManager.cpp b/Engine/Source/Developer/TargetDeviceServices/Private/TargetDeviceServiceManager.cpp
--- a/Engine/Source/Developer/TargetDeviceServices/Private/TargetDeviceServiceManager.cpp
+++ b/Engine/Source/Developer/TargetDeviceServices/Private/TargetDeviceServiceManager.cpp
@@ -140,9 +140,6 @@ void FTargetDeviceServiceManager::InitializeTargetPlatforms()
 {
 	TArray<ITargetPlatform*> Platforms = GetTargetPlatformManager()->GetTargetPlatforms();
 
-	// make sure we're not handling any broadcast events while directly adding devices.
-	FScopeLock Lock(&CriticalSection);
-
 	for (int32 PlatformIndex = 0; PlatformIndex < Platforms.Num(); ++PlatformIndex)
 	{
 		// set up target platform callbacks
