Double target adds are safe, but we must protect against two adds at the same time in every circumstance.

diff --git a/Engine/Source/Developer/TargetDeviceServices/Private/TargetDeviceServiceManager.cpp b/Engine/Source/Developer/TargetDeviceServices/Private/TargetDeviceServiceManager.cpp
--- a/Engine/Source/Developer/TargetDeviceServices/Private/TargetDeviceServiceManager.cpp
+++ b/Engine/Source/Developer/TargetDeviceServices/Private/TargetDeviceServiceManager.cpp
@@ -112,6 +112,7 @@ ITargetDeviceServicePtr FTargetDeviceServiceManager::AddService(const FString& D
 
 bool FTargetDeviceServiceManager::AddTargetDevice(ITargetDevicePtr InDevice)
 {
+	FScopeLock Lock(&CriticalSection);
 	IMessageBusPtr MessageBus = MessageBusPtr.Pin();
 	if (!MessageBus.IsValid())
 	{
@@ -247,6 +248,7 @@ void FTargetDeviceServiceManager::RemoveService(const FString& DeviceName)
 
 void FTargetDeviceServiceManager::RemoveTargetDevice(ITargetDevicePtr InDevice)
 {
+	FScopeLock Lock(&CriticalSection);
 	ITargetDeviceServicePtr DeviceService = DeviceServices.FindRef(InDevice->GetName());
 
 	if (!DeviceService.IsValid())
@@ -340,17 +342,11 @@ void FTargetDeviceServiceManager::HandleMessageBusShutdown()
 
 void FTargetDeviceServiceManager::HandleTargetPlatformDeviceDiscovered(ITargetDeviceRef DiscoveredDevice)
 {
-	FScopeLock Lock(&CriticalSection);
-	{
-		AddTargetDevice(DiscoveredDevice);
-	}
+	AddTargetDevice(DiscoveredDevice);	
 }
 
 
 void FTargetDeviceServiceManager::HandleTargetPlatformDeviceLost(ITargetDeviceRef LostDevice)
 {
-	FScopeLock Lock(&CriticalSection);
-	{
-		RemoveTargetDevice(LostDevice);
-	}
+	RemoveTargetDevice(LostDevice);	
 }
