Back out changelist 2481333

Rob asked me to back out GENERATED_*_BODY -> GENERATED_BODY change for now until the "_Validate and _Implementation auto-generation" discussion is over.

diff --git a/Engine/Plugins/2D/Paper2D/Source/Paper2D/Classes/PaperCharacter.h b/Engine/Plugins/2D/Paper2D/Source/Paper2D/Classes/PaperCharacter.h
--- a/Engine/Plugins/2D/Paper2D/Source/Paper2D/Classes/PaperCharacter.h
+++ b/Engine/Plugins/2D/Paper2D/Source/Paper2D/Classes/PaperCharacter.h
@@ -10,9 +10,7 @@
 UCLASS()
 class PAPER2D_API APaperCharacter : public ACharacter
 {
-	GENERATED_BODY()
-public:
-	APaperCharacter(const FObjectInitializer& ObjectInitializer = FObjectInitializer::Get());
+	GENERATED_UCLASS_BODY()
 	
 	// Name of the Sprite component
 	static FName SpriteComponentName;
diff --git a/Engine/Plugins/2D/Paper2D/Source/Paper2D/Classes/PaperFlipbookActor.h b/Engine/Plugins/2D/Paper2D/Source/Paper2D/Classes/PaperFlipbookActor.h
--- a/Engine/Plugins/2D/Paper2D/Source/Paper2D/Classes/PaperFlipbookActor.h
+++ b/Engine/Plugins/2D/Paper2D/Source/Paper2D/Classes/PaperFlipbookActor.h
@@ -14,9 +14,7 @@
 UCLASS(ComponentWrapperClass)
 class PAPER2D_API APaperFlipbookActor : public AActor
 {
-	GENERATED_BODY()
-public:
-	APaperFlipbookActor(const FObjectInitializer& ObjectInitializer = FObjectInitializer::Get());
+	GENERATED_UCLASS_BODY()
 
 private:
 	UPROPERTY(Category=Sprite, VisibleAnywhere, BlueprintReadOnly, meta=(ExposeFunctionCategories="Sprite,Rendering,Physics,Components|Flipbook", AllowPrivateAccess="true"))
diff --git a/Engine/Plugins/2D/Paper2D/Source/Paper2D/Classes/PaperFlipbookComponent.h b/Engine/Plugins/2D/Paper2D/Source/Paper2D/Classes/PaperFlipbookComponent.h
--- a/Engine/Plugins/2D/Paper2D/Source/Paper2D/Classes/PaperFlipbookComponent.h
+++ b/Engine/Plugins/2D/Paper2D/Source/Paper2D/Classes/PaperFlipbookComponent.h
@@ -11,9 +11,7 @@ DECLARE_DYNAMIC_MULTICAST_DELEGATE(FFlipbookFinishedPlaySignature);
 UCLASS(ShowCategories=(Mobility, ComponentReplication), ClassGroup=Paper2D, EarlyAccessPreview, meta=(BlueprintSpawnableComponent))
 class PAPER2D_API UPaperFlipbookComponent : public UMeshComponent
 {
-	GENERATED_BODY()
-public:
-	UPaperFlipbookComponent(const FObjectInitializer& ObjectInitializer = FObjectInitializer::Get());
+	GENERATED_UCLASS_BODY()
 
 protected:
 	/** Flipbook currently being played */
diff --git a/Engine/Plugins/2D/Paper2D/Source/Paper2D/Classes/PaperSprite.h b/Engine/Plugins/2D/Paper2D/Source/Paper2D/Classes/PaperSprite.h
--- a/Engine/Plugins/2D/Paper2D/Source/Paper2D/Classes/PaperSprite.h
+++ b/Engine/Plugins/2D/Paper2D/Source/Paper2D/Classes/PaperSprite.h
@@ -22,7 +22,7 @@ namespace ESpriteCollisionMode
 USTRUCT()
 struct FPaperSpriteSocket
 {
-	GENERATED_BODY()
+	GENERATED_USTRUCT_BODY()
 
 	// Transform in pivot space (*not* texture space)
 	UPROPERTY(EditAnywhere, BlueprintReadOnly, Category=Sockets)
@@ -44,9 +44,7 @@ typedef TArray<class UTexture2D*, TInlineAllocator<4>> FAdditionalSpriteTextureA
 UCLASS(BlueprintType, meta=(DisplayThumbnail = "true"))
 class PAPER2D_API UPaperSprite : public UObject, public IInterface_CollisionDataProvider
 {
-	GENERATED_BODY()
-public:
-	UPaperSprite(const FObjectInitializer& ObjectInitializer = FObjectInitializer::Get());
+	GENERATED_UCLASS_BODY()
 
 protected:
 
diff --git a/Engine/Plugins/2D/Paper2D/Source/Paper2D/Classes/PaperSpriteActor.h b/Engine/Plugins/2D/Paper2D/Source/Paper2D/Classes/PaperSpriteActor.h
--- a/Engine/Plugins/2D/Paper2D/Source/Paper2D/Classes/PaperSpriteActor.h
+++ b/Engine/Plugins/2D/Paper2D/Source/Paper2D/Classes/PaperSpriteActor.h
@@ -14,9 +14,7 @@
 UCLASS(ComponentWrapperClass, meta = (ChildCanTick))
 class PAPER2D_API APaperSpriteActor : public AActor
 {
-	GENERATED_BODY()
-public:
-	APaperSpriteActor(const FObjectInitializer& ObjectInitializer = FObjectInitializer::Get());
+	GENERATED_UCLASS_BODY()
 
 private:
 	UPROPERTY(Category=Sprite, VisibleAnywhere, BlueprintReadOnly, meta=(ExposeFunctionCategories="Sprite,Rendering,Physics,Components|Sprite", AllowPrivateAccess="true"))
diff --git a/Engine/Plugins/2D/Paper2D/Source/Paper2D/Classes/PaperSpriteAtlas.h b/Engine/Plugins/2D/Paper2D/Source/Paper2D/Classes/PaperSpriteAtlas.h
--- a/Engine/Plugins/2D/Paper2D/Source/Paper2D/Classes/PaperSpriteAtlas.h
+++ b/Engine/Plugins/2D/Paper2D/Source/Paper2D/Classes/PaperSpriteAtlas.h
@@ -7,7 +7,7 @@
 USTRUCT()
 struct FPaperSpriteAtlasSlot
 {
-	GENERATED_BODY()
+	GENERATED_USTRUCT_BODY()
 
 	UPROPERTY()
 	TAssetPtr<class UPaperSprite> SpriteRef;
@@ -47,9 +47,7 @@ enum class EPaperSpriteAtlasPadding : uint8
 UCLASS(MinimalAPI, BlueprintType, Experimental, meta=(DisplayThumbnail = "true"))
 class UPaperSpriteAtlas : public UObject
 {
-	GENERATED_BODY()
-public:
-	PAPER2D_API UPaperSpriteAtlas(const FObjectInitializer& ObjectInitializer = FObjectInitializer::Get());
+	GENERATED_UCLASS_BODY()
 
 #if WITH_EDITORONLY_DATA
 	// Description of this atlas, which shows up in the content browser tooltip
diff --git a/Engine/Plugins/2D/Paper2D/Source/Paper2D/Classes/PaperSpriteComponent.h b/Engine/Plugins/2D/Paper2D/Source/Paper2D/Classes/PaperSpriteComponent.h
--- a/Engine/Plugins/2D/Paper2D/Source/Paper2D/Classes/PaperSpriteComponent.h
+++ b/Engine/Plugins/2D/Paper2D/Source/Paper2D/Classes/PaperSpriteComponent.h
@@ -18,9 +18,7 @@
 UCLASS(ShowCategories=(Mobility), ClassGroup=Paper2D, meta=(BlueprintSpawnableComponent))
 class PAPER2D_API UPaperSpriteComponent : public UMeshComponent
 {
-	GENERATED_BODY()
-public:
-	UPaperSpriteComponent(const FObjectInitializer& ObjectInitializer = FObjectInitializer::Get());
+	GENERATED_UCLASS_BODY()
 
 protected:
 	// The sprite asset used by this component
diff --git a/Engine/Plugins/2D/Paper2D/Source/Paper2D/Classes/PaperTerrainActor.h b/Engine/Plugins/2D/Paper2D/Source/Paper2D/Classes/PaperTerrainActor.h
--- a/Engine/Plugins/2D/Paper2D/Source/Paper2D/Classes/PaperTerrainActor.h
+++ b/Engine/Plugins/2D/Paper2D/Source/Paper2D/Classes/PaperTerrainActor.h
@@ -13,9 +13,7 @@
 UCLASS(BlueprintType, Experimental)
 class PAPER2D_API APaperTerrainActor : public AActor
 {
-	GENERATED_BODY()
-public:
-	APaperTerrainActor(const FObjectInitializer& ObjectInitializer = FObjectInitializer::Get());
+	GENERATED_UCLASS_BODY()
 
 private:
 	UPROPERTY()
diff --git a/Engine/Plugins/2D/Paper2D/Source/Paper2D/Classes/PaperTerrainComponent.h b/Engine/Plugins/2D/Paper2D/Source/Paper2D/Classes/PaperTerrainComponent.h
--- a/Engine/Plugins/2D/Paper2D/Source/Paper2D/Classes/PaperTerrainComponent.h
+++ b/Engine/Plugins/2D/Paper2D/Source/Paper2D/Classes/PaperTerrainComponent.h
@@ -44,9 +44,7 @@ struct FTerrainSegment
 UCLASS(BlueprintType, Experimental)
 class PAPER2D_API UPaperTerrainComponent : public UPrimitiveComponent
 {
-	GENERATED_BODY()
-public:
-	UPaperTerrainComponent(const FObjectInitializer& ObjectInitializer = FObjectInitializer::Get());
+	GENERATED_UCLASS_BODY()
 
 public:
 	/** The terrain material to apply to this component (set of rules for which sprites are used on different surfaces or the interior) */
diff --git a/Engine/Plugins/2D/Paper2D/Source/Paper2D/Classes/PaperTerrainMaterial.h b/Engine/Plugins/2D/Paper2D/Source/Paper2D/Classes/PaperTerrainMaterial.h
--- a/Engine/Plugins/2D/Paper2D/Source/Paper2D/Classes/PaperTerrainMaterial.h
+++ b/Engine/Plugins/2D/Paper2D/Source/Paper2D/Classes/PaperTerrainMaterial.h
@@ -8,7 +8,7 @@
 USTRUCT(BlueprintType)
 struct FPaperTerrainMaterialRule
 {
-	GENERATED_BODY()
+	GENERATED_USTRUCT_BODY()
 
 	// The sprite to use at the 'left' (closest to spline start) edge of the terrain segment
 	UPROPERTY(Category=Sprite, EditAnywhere, meta=(DisplayThumbnail="true"))
@@ -68,9 +68,7 @@ struct FPaperTerrainMaterialRule
 UCLASS(BlueprintType)
 class PAPER2D_API UPaperTerrainMaterial : public UDataAsset
 {
-	GENERATED_BODY()
-public:
-	UPaperTerrainMaterial(const FObjectInitializer& ObjectInitializer = FObjectInitializer::Get());
+	GENERATED_UCLASS_BODY()
 
 	UPROPERTY(Category=Sprite, EditAnywhere)
 	TArray<FPaperTerrainMaterialRule> Rules;
diff --git a/Engine/Plugins/2D/Paper2D/Source/Paper2D/Classes/PaperTerrainSplineComponent.h b/Engine/Plugins/2D/Paper2D/Source/Paper2D/Classes/PaperTerrainSplineComponent.h
--- a/Engine/Plugins/2D/Paper2D/Source/Paper2D/Classes/PaperTerrainSplineComponent.h
+++ b/Engine/Plugins/2D/Paper2D/Source/Paper2D/Classes/PaperTerrainSplineComponent.h
@@ -9,9 +9,7 @@
 UCLASS(BlueprintType, Experimental)
 class PAPER2D_API UPaperTerrainSplineComponent : public USplineComponent
 {
-	GENERATED_BODY()
-public:
-	UPaperTerrainSplineComponent(const FObjectInitializer& ObjectInitializer = FObjectInitializer::Get());
+	GENERATED_UCLASS_BODY()
 
 public:
 	// UObject interface
diff --git a/Engine/Plugins/2D/Paper2D/Source/Paper2D/Classes/PaperTileLayer.h b/Engine/Plugins/2D/Paper2D/Source/Paper2D/Classes/PaperTileLayer.h
--- a/Engine/Plugins/2D/Paper2D/Source/Paper2D/Classes/PaperTileLayer.h
+++ b/Engine/Plugins/2D/Paper2D/Source/Paper2D/Classes/PaperTileLayer.h
@@ -20,7 +20,7 @@ enum class EPaperTileFlags : uint32
 USTRUCT(BlueprintType)
 struct FPaperTileInfo
 {
-	GENERATED_BODY()
+	GENERATED_USTRUCT_BODY()
 
 	// The tile set that this tile comes from
 	UPROPERTY(EditAnywhere, BlueprintReadWrite, Category=Sprite)
@@ -101,9 +101,7 @@ struct FPaperTileInfo
 UCLASS()
 class PAPER2D_API UPaperTileLayer : public UObject
 {
-	GENERATED_BODY()
-public:
-	UPaperTileLayer(const FObjectInitializer& ObjectInitializer = FObjectInitializer::Get());
+	GENERATED_UCLASS_BODY()
 
 	// Name of the layer
 	UPROPERTY()
diff --git a/Engine/Plugins/2D/Paper2D/Source/Paper2D/Classes/PaperTileMap.h b/Engine/Plugins/2D/Paper2D/Source/Paper2D/Classes/PaperTileMap.h
--- a/Engine/Plugins/2D/Paper2D/Source/Paper2D/Classes/PaperTileMap.h
+++ b/Engine/Plugins/2D/Paper2D/Source/Paper2D/Classes/PaperTileMap.h
@@ -31,9 +31,7 @@ namespace ETileMapProjectionMode
 UCLASS()
 class PAPER2D_API UPaperTileMap : public UObject
 {
-	GENERATED_BODY()
-public:
-	UPaperTileMap(const FObjectInitializer& ObjectInitializer = FObjectInitializer::Get());
+	GENERATED_UCLASS_BODY()
 
 	// Width of map (in tiles)
 	UPROPERTY(Category=Setup, EditAnywhere, BlueprintReadOnly, meta=(UIMin=1, ClampMin=1))
diff --git a/Engine/Plugins/2D/Paper2D/Source/Paper2D/Classes/PaperTileMapActor.h b/Engine/Plugins/2D/Paper2D/Source/Paper2D/Classes/PaperTileMapActor.h
--- a/Engine/Plugins/2D/Paper2D/Source/Paper2D/Classes/PaperTileMapActor.h
+++ b/Engine/Plugins/2D/Paper2D/Source/Paper2D/Classes/PaperTileMapActor.h
@@ -15,9 +15,7 @@
 UCLASS(ComponentWrapperClass)
 class PAPER2D_API APaperTileMapActor : public AActor
 {
-	GENERATED_BODY()
-public:
-	APaperTileMapActor(const FObjectInitializer& ObjectInitializer = FObjectInitializer::Get());
+	GENERATED_UCLASS_BODY()
 
 private:
 	UPROPERTY(Category=TileMapActor, VisibleAnywhere, BlueprintReadOnly, meta=(ExposeFunctionCategories="Sprite,Rendering,Physics,Components|Sprite", AllowPrivateAccess="true"))
diff --git a/Engine/Plugins/2D/Paper2D/Source/Paper2D/Classes/PaperTileMapComponent.h b/Engine/Plugins/2D/Paper2D/Source/Paper2D/Classes/PaperTileMapComponent.h
--- a/Engine/Plugins/2D/Paper2D/Source/Paper2D/Classes/PaperTileMapComponent.h
+++ b/Engine/Plugins/2D/Paper2D/Source/Paper2D/Classes/PaperTileMapComponent.h
@@ -20,9 +20,7 @@
 UCLASS(hideCategories=Object, ClassGroup=Paper2D, Experimental, meta=(BlueprintSpawnableComponent))
 class PAPER2D_API UPaperTileMapComponent : public UMeshComponent
 {
-	GENERATED_BODY()
-public:
-	UPaperTileMapComponent(const FObjectInitializer& ObjectInitializer = FObjectInitializer::Get());
+	GENERATED_UCLASS_BODY()
 
 private:
 	UPROPERTY()
diff --git a/Engine/Plugins/2D/Paper2D/Source/Paper2D/Classes/PaperTileSet.h b/Engine/Plugins/2D/Paper2D/Source/Paper2D/Classes/PaperTileSet.h
--- a/Engine/Plugins/2D/Paper2D/Source/Paper2D/Classes/PaperTileSet.h
+++ b/Engine/Plugins/2D/Paper2D/Source/Paper2D/Classes/PaperTileSet.h
@@ -10,9 +10,7 @@
 UCLASS()
 class PAPER2D_API UPaperTileSet : public UObject
 {
-	GENERATED_BODY()
-public:
-	UPaperTileSet(const FObjectInitializer& ObjectInitializer = FObjectInitializer::Get());
+	GENERATED_UCLASS_BODY()
 
 	// The width of a single tile (in pixels)
 	UPROPERTY(Category=TileSet, EditAnywhere, meta=(UIMin=1, ClampMin=1))
diff --git a/Engine/Plugins/2D/Paper2D/Source/Paper2D/Classes/SpriteDrawCall.h b/Engine/Plugins/2D/Paper2D/Source/Paper2D/Classes/SpriteDrawCall.h
--- a/Engine/Plugins/2D/Paper2D/Source/Paper2D/Classes/SpriteDrawCall.h
+++ b/Engine/Plugins/2D/Paper2D/Source/Paper2D/Classes/SpriteDrawCall.h
@@ -8,8 +8,8 @@
 USTRUCT()
 struct FSpriteDrawCallRecord
 {
-	GENERATED_BODY()
 public:
+	GENERATED_USTRUCT_BODY()
 
 	UPROPERTY(Category=Sprite, EditAnywhere)
 	FVector Destination;
diff --git a/Engine/Plugins/2D/Paper2D/Source/Paper2D/Classes/SpriteEditorOnlyTypes.h b/Engine/Plugins/2D/Paper2D/Source/Paper2D/Classes/SpriteEditorOnlyTypes.h
--- a/Engine/Plugins/2D/Paper2D/Source/Paper2D/Classes/SpriteEditorOnlyTypes.h
+++ b/Engine/Plugins/2D/Paper2D/Source/Paper2D/Classes/SpriteEditorOnlyTypes.h
@@ -9,7 +9,7 @@
 USTRUCT()
 struct FSpritePolygon
 {
-	GENERATED_BODY()
+	GENERATED_USTRUCT_BODY()
 
 	FSpritePolygon() 
 	: bNegativeWinding(false)
@@ -82,7 +82,7 @@ namespace ESpriteSubdivisionMode
 USTRUCT()
 struct FSpritePolygonCollection
 {
-	GENERATED_BODY()
+	GENERATED_USTRUCT_BODY()
 
 	UPROPERTY(Category=PolygonData, EditAnywhere, AdvancedDisplay)
 	TArray<FSpritePolygon> Polygons;
@@ -147,7 +147,7 @@ struct FSpritePolygonCollection
 USTRUCT()
 struct FSpriteAssetInitParameters
 {
-	GENERATED_BODY()
+	GENERATED_USTRUCT_BODY()
 
 	FSpriteAssetInitParameters()
 		: Texture(nullptr)
diff --git a/Engine/Plugins/2D/Paper2D/Source/Paper2D/Private/Rendering/PaperBatchComponent.h b/Engine/Plugins/2D/Paper2D/Source/Paper2D/Private/Rendering/PaperBatchComponent.h
--- a/Engine/Plugins/2D/Paper2D/Source/Paper2D/Private/Rendering/PaperBatchComponent.h
+++ b/Engine/Plugins/2D/Paper2D/Source/Paper2D/Private/Rendering/PaperBatchComponent.h
@@ -8,9 +8,7 @@
 UCLASS()
 class UPaperBatchComponent : public UPrimitiveComponent
 {
-	GENERATED_BODY()
-public:
-	UPaperBatchComponent(const FObjectInitializer& ObjectInitializer = FObjectInitializer::Get());
+	GENERATED_UCLASS_BODY()
 
 	// UActorComponent interface
 	virtual void TickComponent(float DeltaTime, enum ELevelTick TickType, FActorComponentTickFunction* ThisTickFunction) override;
diff --git a/Engine/Plugins/2D/Paper2D/Source/Paper2D/Public/MaterialExpressionSpriteTextureSampler.h b/Engine/Plugins/2D/Paper2D/Source/Paper2D/Public/MaterialExpressionSpriteTextureSampler.h
--- a/Engine/Plugins/2D/Paper2D/Source/Paper2D/Public/MaterialExpressionSpriteTextureSampler.h
+++ b/Engine/Plugins/2D/Paper2D/Source/Paper2D/Public/MaterialExpressionSpriteTextureSampler.h
@@ -9,9 +9,7 @@
 UCLASS(hideCategories=(MaterialExpressionTextureSampleParameter, MaterialExpressionSpriteTextureSampler))
 class PAPER2D_API UMaterialExpressionSpriteTextureSampler : public UMaterialExpressionTextureSampleParameter2D
 {
-	GENERATED_BODY()
-public:
-	UMaterialExpressionSpriteTextureSampler(const FObjectInitializer& ObjectInitializer = FObjectInitializer::Get());
+	GENERATED_UCLASS_BODY()
 
 	// Is this a sampler for the default SourceTexture or the AdditionalSourceTextures list?
 	UPROPERTY(EditAnywhere, Category=Paper2D)
diff --git a/Engine/Plugins/2D/Paper2D/Source/Paper2DEditor/Classes/PaperFlipbookActorFactory.h b/Engine/Plugins/2D/Paper2D/Source/Paper2DEditor/Classes/PaperFlipbookActorFactory.h
--- a/Engine/Plugins/2D/Paper2D/Source/Paper2DEditor/Classes/PaperFlipbookActorFactory.h
+++ b/Engine/Plugins/2D/Paper2D/Source/Paper2DEditor/Classes/PaperFlipbookActorFactory.h
@@ -7,9 +7,7 @@
 UCLASS()
 class UPaperFlipbookActorFactory : public UActorFactory
 {
-	GENERATED_BODY()
-public:
-	UPaperFlipbookActorFactory(const FObjectInitializer& ObjectInitializer = FObjectInitializer::Get());
+	GENERATED_UCLASS_BODY()
 
 	// Begin UActorFactory Interface
 	virtual void PostSpawnActor(UObject* Asset, AActor* NewActor) override;
diff --git a/Engine/Plugins/2D/Paper2D/Source/Paper2DEditor/Classes/PaperFlipbookFactory.h b/Engine/Plugins/2D/Paper2D/Source/Paper2DEditor/Classes/PaperFlipbookFactory.h
--- a/Engine/Plugins/2D/Paper2D/Source/Paper2DEditor/Classes/PaperFlipbookFactory.h
+++ b/Engine/Plugins/2D/Paper2D/Source/Paper2DEditor/Classes/PaperFlipbookFactory.h
@@ -10,9 +10,7 @@
 UCLASS()
 class PAPER2DEDITOR_API UPaperFlipbookFactory : public UFactory
 {
-	GENERATED_BODY()
-public:
-	UPaperFlipbookFactory(const FObjectInitializer& ObjectInitializer = FObjectInitializer::Get());
+	GENERATED_UCLASS_BODY()
 
 	TArray<FPaperFlipbookKeyFrame> KeyFrames;
 
diff --git a/Engine/Plugins/2D/Paper2D/Source/Paper2DEditor/Classes/PaperFlipbookThumbnailRenderer.h b/Engine/Plugins/2D/Paper2D/Source/Paper2DEditor/Classes/PaperFlipbookThumbnailRenderer.h
--- a/Engine/Plugins/2D/Paper2D/Source/Paper2DEditor/Classes/PaperFlipbookThumbnailRenderer.h
+++ b/Engine/Plugins/2D/Paper2D/Source/Paper2DEditor/Classes/PaperFlipbookThumbnailRenderer.h
@@ -7,9 +7,7 @@
 UCLASS()
 class UPaperFlipbookThumbnailRenderer : public UPaperSpriteThumbnailRenderer
 {
-	GENERATED_BODY()
-public:
-	UPaperFlipbookThumbnailRenderer(const FObjectInitializer& ObjectInitializer = FObjectInitializer::Get());
+	GENERATED_UCLASS_BODY()
 
 	// UThumbnailRenderer interface
 	virtual void Draw(UObject* Object, int32 X, int32 Y, uint32 Width, uint32 Height, FRenderTarget*, FCanvas* Canvas) override;
diff --git a/Engine/Plugins/2D/Paper2D/Source/Paper2DEditor/Classes/PaperSpriteActorFactory.h b/Engine/Plugins/2D/Paper2D/Source/Paper2DEditor/Classes/PaperSpriteActorFactory.h
--- a/Engine/Plugins/2D/Paper2D/Source/Paper2DEditor/Classes/PaperSpriteActorFactory.h
+++ b/Engine/Plugins/2D/Paper2D/Source/Paper2DEditor/Classes/PaperSpriteActorFactory.h
@@ -7,9 +7,7 @@
 UCLASS()
 class UPaperSpriteActorFactory : public UActorFactory
 {
-	GENERATED_BODY()
-public:
-	UPaperSpriteActorFactory(const FObjectInitializer& ObjectInitializer = FObjectInitializer::Get());
+	GENERATED_UCLASS_BODY()
 
 	// Begin UActorFactory Interface
 	virtual void PostSpawnActor(UObject* Asset, AActor* NewActor) override;
diff --git a/Engine/Plugins/2D/Paper2D/Source/Paper2DEditor/Classes/PaperSpriteFactory.h b/Engine/Plugins/2D/Paper2D/Source/Paper2DEditor/Classes/PaperSpriteFactory.h
--- a/Engine/Plugins/2D/Paper2D/Source/Paper2DEditor/Classes/PaperSpriteFactory.h
+++ b/Engine/Plugins/2D/Paper2D/Source/Paper2DEditor/Classes/PaperSpriteFactory.h
@@ -11,9 +11,7 @@
 UCLASS()
 class UPaperSpriteFactory : public UFactory
 {
-	GENERATED_BODY()
-public:
-	UPaperSpriteFactory(const FObjectInitializer& ObjectInitializer = FObjectInitializer::Get());
+	GENERATED_UCLASS_BODY()
 
 	// Initial texture to create the sprite from (Can be nullptr)
 	class UTexture2D* InitialTexture;
diff --git a/Engine/Plugins/2D/Paper2D/Source/Paper2DEditor/Classes/PaperSpriteThumbnailRenderer.h b/Engine/Plugins/2D/Paper2D/Source/Paper2DEditor/Classes/PaperSpriteThumbnailRenderer.h
--- a/Engine/Plugins/2D/Paper2D/Source/Paper2DEditor/Classes/PaperSpriteThumbnailRenderer.h
+++ b/Engine/Plugins/2D/Paper2D/Source/Paper2DEditor/Classes/PaperSpriteThumbnailRenderer.h
@@ -7,9 +7,7 @@
 UCLASS()
 class UPaperSpriteThumbnailRenderer : public UDefaultSizedThumbnailRenderer
 {
-	GENERATED_BODY()
-public:
-	UPaperSpriteThumbnailRenderer(const FObjectInitializer& ObjectInitializer = FObjectInitializer::Get());
+	GENERATED_UCLASS_BODY()
 
 	// UThumbnailRenderer interface
 	virtual void Draw(UObject* Object, int32 X, int32 Y, uint32 Width, uint32 Height, FRenderTarget*, FCanvas* Canvas) override;
diff --git a/Engine/Plugins/2D/Paper2D/Source/Paper2DEditor/Classes/PaperTileMapFactory.h b/Engine/Plugins/2D/Paper2D/Source/Paper2DEditor/Classes/PaperTileMapFactory.h
--- a/Engine/Plugins/2D/Paper2D/Source/Paper2DEditor/Classes/PaperTileMapFactory.h
+++ b/Engine/Plugins/2D/Paper2D/Source/Paper2DEditor/Classes/PaperTileMapFactory.h
@@ -11,9 +11,7 @@
 UCLASS()
 class UPaperTileMapFactory : public UFactory
 {
-	GENERATED_BODY()
-public:
-	UPaperTileMapFactory(const FObjectInitializer& ObjectInitializer = FObjectInitializer::Get());
+	GENERATED_UCLASS_BODY()
 
 	// Initial tile set to create the tile map from (Can be nullptr)
 	class UPaperTileSet* InitialTileSet;
diff --git a/Engine/Plugins/2D/Paper2D/Source/Paper2DEditor/Classes/PaperTileMapPromotionFactory.h b/Engine/Plugins/2D/Paper2D/Source/Paper2DEditor/Classes/PaperTileMapPromotionFactory.h
--- a/Engine/Plugins/2D/Paper2D/Source/Paper2DEditor/Classes/PaperTileMapPromotionFactory.h
+++ b/Engine/Plugins/2D/Paper2D/Source/Paper2DEditor/Classes/PaperTileMapPromotionFactory.h
@@ -11,9 +11,7 @@
 UCLASS()
 class UPaperTileMapPromotionFactory : public UFactory
 {
-	GENERATED_BODY()
-public:
-	UPaperTileMapPromotionFactory(const FObjectInitializer& ObjectInitializer = FObjectInitializer::Get());
+	GENERATED_UCLASS_BODY()
 
 	// Object being promoted to an asset
 	UPROPERTY()
diff --git a/Engine/Plugins/2D/Paper2D/Source/Paper2DEditor/Classes/PaperTileSetFactory.h b/Engine/Plugins/2D/Paper2D/Source/Paper2DEditor/Classes/PaperTileSetFactory.h
--- a/Engine/Plugins/2D/Paper2D/Source/Paper2DEditor/Classes/PaperTileSetFactory.h
+++ b/Engine/Plugins/2D/Paper2D/Source/Paper2DEditor/Classes/PaperTileSetFactory.h
@@ -11,9 +11,7 @@
 UCLASS()
 class UPaperTileSetFactory : public UFactory
 {
-	GENERATED_BODY()
-public:
-	UPaperTileSetFactory(const FObjectInitializer& ObjectInitializer = FObjectInitializer::Get());
+	GENERATED_UCLASS_BODY()
 
 	// Initial texture to create the tile set from (Can be nullptr)
 	class UTexture2D* InitialTexture;
diff --git a/Engine/Plugins/2D/Paper2D/Source/Paper2DEditor/Classes/PaperTileSetThumbnailRenderer.h b/Engine/Plugins/2D/Paper2D/Source/Paper2DEditor/Classes/PaperTileSetThumbnailRenderer.h
--- a/Engine/Plugins/2D/Paper2D/Source/Paper2DEditor/Classes/PaperTileSetThumbnailRenderer.h
+++ b/Engine/Plugins/2D/Paper2D/Source/Paper2DEditor/Classes/PaperTileSetThumbnailRenderer.h
@@ -7,9 +7,7 @@
 UCLASS()
 class UPaperTileSetThumbnailRenderer : public UDefaultSizedThumbnailRenderer
 {
-	GENERATED_BODY()
-public:
-	UPaperTileSetThumbnailRenderer(const FObjectInitializer& ObjectInitializer = FObjectInitializer::Get());
+	GENERATED_UCLASS_BODY()
 
 	// UThumbnailRenderer interface
 	virtual void Draw(UObject* Object, int32 X, int32 Y, uint32 Width, uint32 Height, FRenderTarget*, FCanvas* Canvas) override;
diff --git a/Engine/Plugins/2D/Paper2D/Source/Paper2DEditor/Classes/SpriteEditorTools/PSE_DeleteTool.h b/Engine/Plugins/2D/Paper2D/Source/Paper2DEditor/Classes/SpriteEditorTools/PSE_DeleteTool.h
--- a/Engine/Plugins/2D/Paper2D/Source/Paper2DEditor/Classes/SpriteEditorTools/PSE_DeleteTool.h
+++ b/Engine/Plugins/2D/Paper2D/Source/Paper2DEditor/Classes/SpriteEditorTools/PSE_DeleteTool.h
@@ -7,7 +7,5 @@
 UCLASS()
 class UPSE_DeleteTool : public USpriteEditorTool
 {
-	GENERATED_BODY()
-public:
-	UPSE_DeleteTool(const FObjectInitializer& ObjectInitializer = FObjectInitializer::Get());
+	GENERATED_UCLASS_BODY()
 };
diff --git a/Engine/Plugins/2D/Paper2D/Source/Paper2DEditor/Classes/SpriteEditorTools/PSE_LineTool.h b/Engine/Plugins/2D/Paper2D/Source/Paper2DEditor/Classes/SpriteEditorTools/PSE_LineTool.h
--- a/Engine/Plugins/2D/Paper2D/Source/Paper2DEditor/Classes/SpriteEditorTools/PSE_LineTool.h
+++ b/Engine/Plugins/2D/Paper2D/Source/Paper2DEditor/Classes/SpriteEditorTools/PSE_LineTool.h
@@ -7,7 +7,5 @@
 UCLASS()
 class UPSE_LineTool : public USpriteEditorTool
 {
-	GENERATED_BODY()
-public:
-	UPSE_LineTool(const FObjectInitializer& ObjectInitializer = FObjectInitializer::Get());
+	GENERATED_UCLASS_BODY()
 };
diff --git a/Engine/Plugins/2D/Paper2D/Source/Paper2DEditor/Classes/SpriteEditorTools/PSE_SplitEdgeTool.h b/Engine/Plugins/2D/Paper2D/Source/Paper2DEditor/Classes/SpriteEditorTools/PSE_SplitEdgeTool.h
--- a/Engine/Plugins/2D/Paper2D/Source/Paper2DEditor/Classes/SpriteEditorTools/PSE_SplitEdgeTool.h
+++ b/Engine/Plugins/2D/Paper2D/Source/Paper2DEditor/Classes/SpriteEditorTools/PSE_SplitEdgeTool.h
@@ -7,7 +7,5 @@
 UCLASS()
 class UPSE_SplitEdgeTool : public USpriteEditorTool
 {
-	GENERATED_BODY()
-public:
-	UPSE_SplitEdgeTool(const FObjectInitializer& ObjectInitializer = FObjectInitializer::Get());
+	GENERATED_UCLASS_BODY()
 };
diff --git a/Engine/Plugins/2D/Paper2D/Source/Paper2DEditor/Classes/SpriteEditorTools/SpriteEditorTool.h b/Engine/Plugins/2D/Paper2D/Source/Paper2DEditor/Classes/SpriteEditorTools/SpriteEditorTool.h
--- a/Engine/Plugins/2D/Paper2D/Source/Paper2DEditor/Classes/SpriteEditorTools/SpriteEditorTool.h
+++ b/Engine/Plugins/2D/Paper2D/Source/Paper2DEditor/Classes/SpriteEditorTools/SpriteEditorTool.h
@@ -7,7 +7,5 @@
 UCLASS()
 class USpriteEditorTool : public UObject
 {
-	GENERATED_BODY()
-public:
-	USpriteEditorTool(const FObjectInitializer& ObjectInitializer = FObjectInitializer::Get());
+	GENERATED_UCLASS_BODY()
 };
diff --git a/Engine/Plugins/2D/Paper2D/Source/Paper2DEditor/Classes/TerrainSplineActorFactory.h b/Engine/Plugins/2D/Paper2D/Source/Paper2DEditor/Classes/TerrainSplineActorFactory.h
--- a/Engine/Plugins/2D/Paper2D/Source/Paper2DEditor/Classes/TerrainSplineActorFactory.h
+++ b/Engine/Plugins/2D/Paper2D/Source/Paper2DEditor/Classes/TerrainSplineActorFactory.h
@@ -7,9 +7,7 @@
 UCLASS()
 class UTerrainSplineActorFactory : public UActorFactory
 {
-	GENERATED_BODY()
-public:
-	UTerrainSplineActorFactory(const FObjectInitializer& ObjectInitializer = FObjectInitializer::Get());
+	GENERATED_UCLASS_BODY()
 
 	// UActorFactory interface
 	virtual void PostSpawnActor(UObject* Asset, AActor* NewActor) override;
diff --git a/Engine/Plugins/2D/Paper2D/Source/Paper2DEditor/Classes/TileMapActorFactory.h b/Engine/Plugins/2D/Paper2D/Source/Paper2DEditor/Classes/TileMapActorFactory.h
--- a/Engine/Plugins/2D/Paper2D/Source/Paper2DEditor/Classes/TileMapActorFactory.h
+++ b/Engine/Plugins/2D/Paper2D/Source/Paper2DEditor/Classes/TileMapActorFactory.h
@@ -7,9 +7,7 @@
 UCLASS()
 class UTileMapActorFactory : public UActorFactory
 {
-	GENERATED_BODY()
-public:
-	UTileMapActorFactory(const FObjectInitializer& ObjectInitializer = FObjectInitializer::Get());
+	GENERATED_UCLASS_BODY()
 
 	// UActorFactory interface
 	virtual void PostSpawnActor(UObject* Asset, AActor* NewActor) override;
diff --git a/Engine/Plugins/2D/Paper2D/Source/Paper2DEditor/Classes/TileMapAssetImportData.h b/Engine/Plugins/2D/Paper2D/Source/Paper2DEditor/Classes/TileMapAssetImportData.h
--- a/Engine/Plugins/2D/Paper2D/Source/Paper2DEditor/Classes/TileMapAssetImportData.h
+++ b/Engine/Plugins/2D/Paper2D/Source/Paper2DEditor/Classes/TileMapAssetImportData.h
@@ -7,7 +7,7 @@
 USTRUCT()
 struct FTileSetImportMapping
 {
-	GENERATED_BODY()
+	GENERATED_USTRUCT_BODY()
 
 	UPROPERTY()
 	FString SourceName;
@@ -25,9 +25,7 @@ struct FTileSetImportMapping
 UCLASS()
 class PAPER2DEDITOR_API UTileMapAssetImportData : public UAssetImportData
 {
-	GENERATED_BODY()
-public:
-	UTileMapAssetImportData(const FObjectInitializer& ObjectInitializer = FObjectInitializer::Get());
+	GENERATED_UCLASS_BODY()
 
 	UPROPERTY()
 	TArray<FTileSetImportMapping> TileSetMap;
diff --git a/Engine/Plugins/2D/Paper2D/Source/PaperJsonImporter/Classes/PaperJsonImporterFactory.h b/Engine/Plugins/2D/Paper2D/Source/PaperJsonImporter/Classes/PaperJsonImporterFactory.h
--- a/Engine/Plugins/2D/Paper2D/Source/PaperJsonImporter/Classes/PaperJsonImporterFactory.h
+++ b/Engine/Plugins/2D/Paper2D/Source/PaperJsonImporter/Classes/PaperJsonImporterFactory.h
@@ -8,9 +8,7 @@
 UCLASS()
 class UPaperJsonImporterFactory : public UFactory
 {
-	GENERATED_BODY()
-public:
-	UPaperJsonImporterFactory(const FObjectInitializer& ObjectInitializer = FObjectInitializer::Get());
+	GENERATED_UCLASS_BODY()
 
 	// UFactory interface
 	virtual FText GetToolTip() const override;
diff --git a/Engine/Plugins/2D/Paper2D/Source/PaperJsonImporter/Private/PaperSpriteSheet.h b/Engine/Plugins/2D/Paper2D/Source/PaperJsonImporter/Private/PaperSpriteSheet.h
--- a/Engine/Plugins/2D/Paper2D/Source/PaperJsonImporter/Private/PaperSpriteSheet.h
+++ b/Engine/Plugins/2D/Paper2D/Source/PaperJsonImporter/Private/PaperSpriteSheet.h
@@ -7,9 +7,7 @@
 UCLASS(BlueprintType, meta = (DisplayThumbnail = "true"))
 class UPaperSpriteSheet : public UObject
 {
-	GENERATED_BODY()
-public:
-	UPaperSpriteSheet(const FObjectInitializer& ObjectInitializer = FObjectInitializer::Get());
+	GENERATED_UCLASS_BODY()
 
 public:
 	// The names of sprites during import
diff --git a/Engine/Plugins/Developer/UObjectPlugin/Source/UObjectPlugin/Classes/MyPluginObject.h b/Engine/Plugins/Developer/UObjectPlugin/Source/UObjectPlugin/Classes/MyPluginObject.h
--- a/Engine/Plugins/Developer/UObjectPlugin/Source/UObjectPlugin/Classes/MyPluginObject.h
+++ b/Engine/Plugins/Developer/UObjectPlugin/Source/UObjectPlugin/Classes/MyPluginObject.h
@@ -11,7 +11,7 @@
 USTRUCT()
 struct FMyPluginStruct
 {
-	GENERATED_BODY()
+	GENERATED_USTRUCT_BODY()
  
 	UPROPERTY()
 	FString TestString;
@@ -24,9 +24,7 @@ struct FMyPluginStruct
 UCLASS()
 class UMyPluginObject : public UObject
 {
-	GENERATED_BODY()
-public:
-	UMyPluginObject(const FObjectInitializer& ObjectInitializer = FObjectInitializer::Get());
+	GENERATED_UCLASS_BODY()
 
 public:
 
diff --git a/Engine/Plugins/Editor/MacGraphicsSwitching/Source/MacGraphicsSwitching/Private/MacGraphicsSwitchingSettings.h b/Engine/Plugins/Editor/MacGraphicsSwitching/Source/MacGraphicsSwitching/Private/MacGraphicsSwitchingSettings.h
--- a/Engine/Plugins/Editor/MacGraphicsSwitching/Source/MacGraphicsSwitching/Private/MacGraphicsSwitchingSettings.h
+++ b/Engine/Plugins/Editor/MacGraphicsSwitching/Source/MacGraphicsSwitching/Private/MacGraphicsSwitchingSettings.h
@@ -6,10 +6,9 @@
 UCLASS(Config=EditorGameAgnostic)
 class UMacGraphicsSwitchingSettings : public UObject
 {
-	GENERATED_BODY()
+	GENERATED_UCLASS_BODY()
 	
 public:
-	UMacGraphicsSwitchingSettings(const FObjectInitializer& ObjectInitializer = FObjectInitializer::Get());
 	
 	UPROPERTY(Config, EditAnywhere, Category=OpenGL)
 	int32 DisplayMask;
diff --git a/Engine/Plugins/Editor/SpeedTreeImporter/Source/SpeedTreeImporter/Classes/ReimportSpeedTreeFactory.h b/Engine/Plugins/Editor/SpeedTreeImporter/Source/SpeedTreeImporter/Classes/ReimportSpeedTreeFactory.h
--- a/Engine/Plugins/Editor/SpeedTreeImporter/Source/SpeedTreeImporter/Classes/ReimportSpeedTreeFactory.h
+++ b/Engine/Plugins/Editor/SpeedTreeImporter/Source/SpeedTreeImporter/Classes/ReimportSpeedTreeFactory.h
@@ -11,9 +11,7 @@
 UCLASS(collapsecategories)
 class UReimportSpeedTreeFactory : public USpeedTreeImportFactory, public FReimportHandler
 {
-	GENERATED_BODY()
-public:
-	UReimportSpeedTreeFactory(const FObjectInitializer& ObjectInitializer = FObjectInitializer::Get());
+	GENERATED_UCLASS_BODY()
 
 	// Begin FReimportHandler interface
 	virtual bool CanReimport(UObject* Obj, TArray<FString>& OutFilenames) override;
diff --git a/Engine/Plugins/Editor/SpeedTreeImporter/Source/SpeedTreeImporter/Classes/SpeedTreeImportFactory.h b/Engine/Plugins/Editor/SpeedTreeImporter/Source/SpeedTreeImporter/Classes/SpeedTreeImportFactory.h
--- a/Engine/Plugins/Editor/SpeedTreeImporter/Source/SpeedTreeImporter/Classes/SpeedTreeImportFactory.h
+++ b/Engine/Plugins/Editor/SpeedTreeImporter/Source/SpeedTreeImporter/Classes/SpeedTreeImportFactory.h
@@ -10,9 +10,7 @@
 UCLASS(hidecategories=Object)
 class USpeedTreeImportFactory : public UFactory
 {
-	GENERATED_BODY()
-public:
-	USpeedTreeImportFactory(const FObjectInitializer& ObjectInitializer = FObjectInitializer::Get());
+	GENERATED_UCLASS_BODY()
 
 	// Begin UFactory Interface
 	virtual FText GetDisplayName() const override;
diff --git a/Engine/Plugins/Experimental/HTML5Networking/Source/HTML5Networking/Classes/WebSocketConnection.h b/Engine/Plugins/Experimental/HTML5Networking/Source/HTML5Networking/Classes/WebSocketConnection.h
--- a/Engine/Plugins/Experimental/HTML5Networking/Source/HTML5Networking/Classes/WebSocketConnection.h
+++ b/Engine/Plugins/Experimental/HTML5Networking/Source/HTML5Networking/Classes/WebSocketConnection.h
@@ -8,9 +8,7 @@
 UCLASS(transient, config = Engine)
 class HTML5NETWORKING_API UWebSocketConnection : public UNetConnection
 {
-	GENERATED_BODY()
-public:
-	UWebSocketConnection(const FObjectInitializer& ObjectInitializer = FObjectInitializer::Get());
+	GENERATED_UCLASS_BODY()
 
 	class FWebSocket*			WebSocket; 
 
diff --git a/Engine/Plugins/Experimental/HTML5Networking/Source/HTML5Networking/Classes/WebSocketNetDriver.h b/Engine/Plugins/Experimental/HTML5Networking/Source/HTML5Networking/Classes/WebSocketNetDriver.h
--- a/Engine/Plugins/Experimental/HTML5Networking/Source/HTML5Networking/Classes/WebSocketNetDriver.h
+++ b/Engine/Plugins/Experimental/HTML5Networking/Source/HTML5Networking/Classes/WebSocketNetDriver.h
@@ -9,9 +9,7 @@
 UCLASS(transient, config = Engine)
 class HTML5NETWORKING_API UWebSocketNetDriver : public UNetDriver
 {
-	GENERATED_BODY()
-public:
-	UWebSocketNetDriver(const FObjectInitializer& ObjectInitializer = FObjectInitializer::Get());
+	GENERATED_UCLASS_BODY()
 
 	/** Websocket server port*/
 	UPROPERTY(Config)
diff --git a/Engine/Plugins/Experimental/LiveEditor/Source/LiveEditor/Classes/K2Node_LiveEditObject.h b/Engine/Plugins/Experimental/LiveEditor/Source/LiveEditor/Classes/K2Node_LiveEditObject.h
--- a/Engine/Plugins/Experimental/LiveEditor/Source/LiveEditor/Classes/K2Node_LiveEditObject.h
+++ b/Engine/Plugins/Experimental/LiveEditor/Source/LiveEditor/Classes/K2Node_LiveEditObject.h
@@ -9,9 +9,7 @@
 UCLASS(MinimalAPI)
 class UK2Node_LiveEditObject : public UK2Node
 {
-	GENERATED_BODY()
-public:
-	LIVEEDITOR_API UK2Node_LiveEditObject(const FObjectInitializer& ObjectInitializer = FObjectInitializer::Get());
+	GENERATED_UCLASS_BODY()
 
 #if WITH_EDITOR
 	// Begin UEdGraphNode interface.
diff --git a/Engine/Plugins/Experimental/LiveEditor/Source/LiveEditor/Classes/LiveEditorBlueprint.h b/Engine/Plugins/Experimental/LiveEditor/Source/LiveEditor/Classes/LiveEditorBlueprint.h
--- a/Engine/Plugins/Experimental/LiveEditor/Source/LiveEditor/Classes/LiveEditorBlueprint.h
+++ b/Engine/Plugins/Experimental/LiveEditor/Source/LiveEditor/Classes/LiveEditorBlueprint.h
@@ -8,9 +8,7 @@
 UCLASS(Abstract, Blueprintable)
 class ULiveEditorBlueprint : public UObject
 {
-	GENERATED_BODY()
-public:
-	ULiveEditorBlueprint(const FObjectInitializer& ObjectInitializer = FObjectInitializer::Get());
+	GENERATED_UCLASS_BODY()
 
 public:
 	UFUNCTION(BlueprintImplementableEvent)
diff --git a/Engine/Plugins/Experimental/LiveEditor/Source/LiveEditor/Classes/LiveEditorBroadcaster.h b/Engine/Plugins/Experimental/LiveEditor/Source/LiveEditor/Classes/LiveEditorBroadcaster.h
--- a/Engine/Plugins/Experimental/LiveEditor/Source/LiveEditor/Classes/LiveEditorBroadcaster.h
+++ b/Engine/Plugins/Experimental/LiveEditor/Source/LiveEditor/Classes/LiveEditorBroadcaster.h
@@ -11,9 +11,7 @@ DECLARE_DYNAMIC_MULTICAST_DELEGATE_ThreeParams( FLiveEditorEventMIDIMultiCast, i
 UCLASS(MinimalAPI)
 class ULiveEditorBroadcaster : public UObject
 {
-	GENERATED_BODY()
-public:
-	LIVEEDITOR_API ULiveEditorBroadcaster(const FObjectInitializer& ObjectInitializer = FObjectInitializer::Get());
+	GENERATED_UCLASS_BODY()
 
 	UPROPERTY(BlueprintAssignable)
 	FLiveEditorEventMIDIMultiCast OnEventMIDI;
diff --git a/Engine/Plugins/Experimental/LiveEditor/Source/LiveEditor/Classes/LiveEditorKismetLibrary.h b/Engine/Plugins/Experimental/LiveEditor/Source/LiveEditor/Classes/LiveEditorKismetLibrary.h
--- a/Engine/Plugins/Experimental/LiveEditor/Source/LiveEditor/Classes/LiveEditorKismetLibrary.h
+++ b/Engine/Plugins/Experimental/LiveEditor/Source/LiveEditor/Classes/LiveEditorKismetLibrary.h
@@ -9,9 +9,7 @@
 UCLASS()
 class ULiveEditorKismetLibrary : public UBlueprintFunctionLibrary
 {
-	GENERATED_BODY()
-public:
-	ULiveEditorKismetLibrary(const FObjectInitializer& ObjectInitializer = FObjectInitializer::Get());
+	GENERATED_UCLASS_BODY()
 
 	UFUNCTION(BlueprintPure, Category="Development|LiveEditor")
 	static float LiveModifyFloat( float InFloat );
diff --git a/Engine/Plugins/Experimental/LiveEditor/Source/LiveEditor/Classes/LiveEditorTypes.h b/Engine/Plugins/Experimental/LiveEditor/Source/LiveEditor/Classes/LiveEditorTypes.h
--- a/Engine/Plugins/Experimental/LiveEditor/Source/LiveEditor/Classes/LiveEditorTypes.h
+++ b/Engine/Plugins/Experimental/LiveEditor/Source/LiveEditor/Classes/LiveEditorTypes.h
@@ -20,7 +20,7 @@ namespace ELiveEditControllerType UMETA(BlueprintType="true")
 USTRUCT()
 struct FLiveEditBinding
 {
-	GENERATED_BODY()
+	GENERATED_USTRUCT_BODY()
 
 	UPROPERTY()
 	int32 Channel;
@@ -51,7 +51,7 @@ public:
 USTRUCT(BlueprintType)
 struct FLiveEditorCheckpointData
 {
-	GENERATED_BODY()
+	GENERATED_USTRUCT_BODY()
 
 	UPROPERTY(EditAnywhere, BlueprintReadWrite, Category="Live Editor")
 	int32 IntValue;
@@ -65,11 +65,10 @@ struct FLiveEditorCheckpointData
 	}
 };
 
-UCLASS(Abstract)
+UCLASS(Abstract, CustomConstructor)
 class ULiveEditorTypes : public UObject
 {
-	GENERATED_BODY()
-public:
+	GENERATED_UCLASS_BODY()
 
 	ULiveEditorTypes(const FObjectInitializer& ObjectInitializer) : Super(ObjectInitializer) {}
 };
diff --git a/Engine/Plugins/Media/MediaPlayerEditor/Source/MediaPlayerEditor/Private/Factories/MediaPlayerFactory.h b/Engine/Plugins/Media/MediaPlayerEditor/Source/MediaPlayerEditor/Private/Factories/MediaPlayerFactory.h
--- a/Engine/Plugins/Media/MediaPlayerEditor/Source/MediaPlayerEditor/Private/Factories/MediaPlayerFactory.h
+++ b/Engine/Plugins/Media/MediaPlayerEditor/Source/MediaPlayerEditor/Private/Factories/MediaPlayerFactory.h
@@ -12,9 +12,7 @@ UCLASS(hidecategories=Object)
 class UMediaPlayerFactory
 	: public UFactory
 {
-	GENERATED_BODY()
-public:
-	UMediaPlayerFactory(const FObjectInitializer& ObjectInitializer = FObjectInitializer::Get());
+	GENERATED_UCLASS_BODY()
 
 public:
 
diff --git a/Engine/Plugins/Media/MediaPlayerEditor/Source/MediaPlayerEditor/Private/Factories/MediaSoundWaveFactoryNew.h b/Engine/Plugins/Media/MediaPlayerEditor/Source/MediaPlayerEditor/Private/Factories/MediaSoundWaveFactoryNew.h
--- a/Engine/Plugins/Media/MediaPlayerEditor/Source/MediaPlayerEditor/Private/Factories/MediaSoundWaveFactoryNew.h
+++ b/Engine/Plugins/Media/MediaPlayerEditor/Source/MediaPlayerEditor/Private/Factories/MediaSoundWaveFactoryNew.h
@@ -12,9 +12,7 @@ UCLASS(hidecategories=Object, MinimalAPI)
 class UMediaSoundWaveFactoryNew
 	: public UFactory
 {
-	GENERATED_BODY()
-public:
-	MEDIAPLAYEREDITOR_API UMediaSoundWaveFactoryNew(const FObjectInitializer& ObjectInitializer = FObjectInitializer::Get());
+	GENERATED_UCLASS_BODY()
 
 	/** An initial UMediaPlayer asset to place in the newly created sound wave. */
 	UPROPERTY()
diff --git a/Engine/Plugins/Media/MediaPlayerEditor/Source/MediaPlayerEditor/Private/Factories/MediaTextureFactoryNew.h b/Engine/Plugins/Media/MediaPlayerEditor/Source/MediaPlayerEditor/Private/Factories/MediaTextureFactoryNew.h
--- a/Engine/Plugins/Media/MediaPlayerEditor/Source/MediaPlayerEditor/Private/Factories/MediaTextureFactoryNew.h
+++ b/Engine/Plugins/Media/MediaPlayerEditor/Source/MediaPlayerEditor/Private/Factories/MediaTextureFactoryNew.h
@@ -12,9 +12,7 @@ UCLASS(hidecategories=Object, MinimalAPI)
 class UMediaTextureFactoryNew
 	: public UFactory
 {
-	GENERATED_BODY()
-public:
-	MEDIAPLAYEREDITOR_API UMediaTextureFactoryNew(const FObjectInitializer& ObjectInitializer = FObjectInitializer::Get());
+	GENERATED_UCLASS_BODY()
 
 	/** An initial media player asset to place in the newly created texture. */
 	UPROPERTY()
diff --git a/Engine/Plugins/Messaging/ShmMessaging/Source/ShmMessaging/Private/Shared/ShmMessagingSettings.h b/Engine/Plugins/Messaging/ShmMessaging/Source/ShmMessaging/Private/Shared/ShmMessagingSettings.h
--- a/Engine/Plugins/Messaging/ShmMessaging/Source/ShmMessaging/Private/Shared/ShmMessagingSettings.h
+++ b/Engine/Plugins/Messaging/ShmMessaging/Source/ShmMessaging/Private/Shared/ShmMessagingSettings.h
@@ -9,9 +9,7 @@ UCLASS(config=Engine)
 class UShmMessagingSettings
 	: public UObject
 {
-	GENERATED_BODY()
-public:
-	UShmMessagingSettings(const FObjectInitializer& ObjectInitializer = FObjectInitializer::Get());
+	GENERATED_UCLASS_BODY()
 
 public:
 
diff --git a/Engine/Plugins/Messaging/UdpMessaging/Source/UdpMessaging/Private/Shared/UdpMessagingSettings.h b/Engine/Plugins/Messaging/UdpMessaging/Source/UdpMessaging/Private/Shared/UdpMessagingSettings.h
--- a/Engine/Plugins/Messaging/UdpMessaging/Source/UdpMessaging/Private/Shared/UdpMessagingSettings.h
+++ b/Engine/Plugins/Messaging/UdpMessaging/Source/UdpMessaging/Private/Shared/UdpMessagingSettings.h
@@ -9,9 +9,7 @@ UCLASS(config=Engine)
 class UUdpMessagingSettings
 	: public UObject
 {
-	GENERATED_BODY()
-public:
-	UUdpMessagingSettings(const FObjectInitializer& ObjectInitializer = FObjectInitializer::Get());
+	GENERATED_UCLASS_BODY()
 
 public:
 
diff --git a/Engine/Plugins/Runtime/Analytics/AnalyticsBlueprintLibrary/Source/AnalyticsBlueprintLibrary/Classes/AnalyticsBlueprintLibrary.h b/Engine/Plugins/Runtime/Analytics/AnalyticsBlueprintLibrary/Source/AnalyticsBlueprintLibrary/Classes/AnalyticsBlueprintLibrary.h
--- a/Engine/Plugins/Runtime/Analytics/AnalyticsBlueprintLibrary/Source/AnalyticsBlueprintLibrary/Classes/AnalyticsBlueprintLibrary.h
+++ b/Engine/Plugins/Runtime/Analytics/AnalyticsBlueprintLibrary/Source/AnalyticsBlueprintLibrary/Classes/AnalyticsBlueprintLibrary.h
@@ -8,7 +8,7 @@
 USTRUCT(BlueprintType)
 struct FAnalyticsEventAttr
 {
-	GENERATED_BODY()
+	GENERATED_USTRUCT_BODY();
 
 	UPROPERTY(EditAnywhere, BlueprintReadWrite, Category="Analytics")
 	FString Name;
@@ -21,9 +21,7 @@ UCLASS()
 class UAnalyticsBlueprintLibrary :
 	public UBlueprintFunctionLibrary
 {
-	GENERATED_BODY()
-public:
-	UAnalyticsBlueprintLibrary(const FObjectInitializer& ObjectInitializer = FObjectInitializer::Get());
+	GENERATED_UCLASS_BODY()
 
 	/** Starts an analytics session without any custom attributes specified */
 	UFUNCTION(BlueprintCallable, Category="Analytics")
diff --git a/Engine/Plugins/Runtime/CableComponent/Source/CableComponent/Classes/CableActor.h b/Engine/Plugins/Runtime/CableComponent/Source/CableComponent/Classes/CableActor.h
--- a/Engine/Plugins/Runtime/CableComponent/Source/CableComponent/Classes/CableActor.h
+++ b/Engine/Plugins/Runtime/CableComponent/Source/CableComponent/Classes/CableActor.h
@@ -8,9 +8,7 @@
 UCLASS(hidecategories=(Input,Collision,Replication), showcategories=("Input|MouseInput", "Input|TouchInput"))
 class ACableActor : public AActor
 {
-	GENERATED_BODY()
-public:
-	ACableActor(const FObjectInitializer& ObjectInitializer = FObjectInitializer::Get());
+	GENERATED_UCLASS_BODY()
 
 	/** Cable component that performs simulation and rendering */
 	UPROPERTY(Category=Cable, VisibleAnywhere, BlueprintReadOnly)
diff --git a/Engine/Plugins/Runtime/CableComponent/Source/CableComponent/Classes/CableComponent.h b/Engine/Plugins/Runtime/CableComponent/Source/CableComponent/Classes/CableComponent.h
--- a/Engine/Plugins/Runtime/CableComponent/Source/CableComponent/Classes/CableComponent.h
+++ b/Engine/Plugins/Runtime/CableComponent/Source/CableComponent/Classes/CableComponent.h
@@ -25,9 +25,7 @@ struct FCableParticle
 UCLASS(hidecategories=(Object, Physics, Collision, Activation, "Components|Activation"), editinlinenew, meta=(BlueprintSpawnableComponent), ClassGroup=Rendering)
 class UCableComponent : public UMeshComponent
 {
-	GENERATED_BODY()
-public:
-	UCableComponent(const FObjectInitializer& ObjectInitializer = FObjectInitializer::Get());
+	GENERATED_UCLASS_BODY()
 
 public:
 
diff --git a/Engine/Plugins/Runtime/LeapMotionController/Source/LeapMotionController/Public/LeapMotionBoneActor.h b/Engine/Plugins/Runtime/LeapMotionController/Source/LeapMotionController/Public/LeapMotionBoneActor.h
--- a/Engine/Plugins/Runtime/LeapMotionController/Source/LeapMotionController/Public/LeapMotionBoneActor.h
+++ b/Engine/Plugins/Runtime/LeapMotionController/Source/LeapMotionController/Public/LeapMotionBoneActor.h
@@ -16,9 +16,7 @@ class ULeapMotionControllerComponent;
 UCLASS(ClassGroup = LeapMotion, BlueprintType, Blueprintable)
 class LEAPMOTIONCONTROLLER_API ALeapMotionBoneActor : public AActor
 {
-	GENERATED_BODY()
-public:
-	ALeapMotionBoneActor(const FObjectInitializer& ObjectInitializer = FObjectInitializer::Get());
+	GENERATED_UCLASS_BODY()
 
 	UPROPERTY(VisibleAnywhere, BlueprintReadOnly, Category = LeapMotion)
 	ELeapBone LeapBoneType;
diff --git a/Engine/Plugins/Runtime/LeapMotionController/Source/LeapMotionController/Public/LeapMotionControllerComponent.h b/Engine/Plugins/Runtime/LeapMotionController/Source/LeapMotionController/Public/LeapMotionControllerComponent.h
--- a/Engine/Plugins/Runtime/LeapMotionController/Source/LeapMotionController/Public/LeapMotionControllerComponent.h
+++ b/Engine/Plugins/Runtime/LeapMotionController/Source/LeapMotionController/Public/LeapMotionControllerComponent.h
@@ -19,9 +19,7 @@ class ALeapMotionHandActor;
 UCLASS(ClassGroup = LeapMotion, BlueprintType, Blueprintable, meta = (BlueprintSpawnableComponent))
 class LEAPMOTIONCONTROLLER_API ULeapMotionControllerComponent : public USceneComponent
 {
-	GENERATED_BODY()
-public:
-	ULeapMotionControllerComponent(const FObjectInitializer& ObjectInitializer = FObjectInitializer::Get());
+	GENERATED_UCLASS_BODY()
 
 	/** Scale of hands, relatively to real-world size. This ignores actor's transform's scale */
 	UPROPERTY(EditAnywhere, BlueprintReadWrite, Category = LeapMotion)
diff --git a/Engine/Plugins/Runtime/LeapMotionController/Source/LeapMotionController/Public/LeapMotionFunctionLibrary.h b/Engine/Plugins/Runtime/LeapMotionController/Source/LeapMotionController/Public/LeapMotionFunctionLibrary.h
--- a/Engine/Plugins/Runtime/LeapMotionController/Source/LeapMotionController/Public/LeapMotionFunctionLibrary.h
+++ b/Engine/Plugins/Runtime/LeapMotionController/Source/LeapMotionController/Public/LeapMotionFunctionLibrary.h
@@ -18,9 +18,7 @@ class ULeapMotionControllerComponent;
 UCLASS()
 class LEAPMOTIONCONTROLLER_API ULeapMotionFunctionLibrary : public UBlueprintFunctionLibrary
 {
-	GENERATED_BODY()
-public:
-	ULeapMotionFunctionLibrary(const FObjectInitializer& ObjectInitializer = FObjectInitializer::Get());
+	GENERATED_UCLASS_BODY()
 
 	/** 
 	 * Checks whether a Leap Motion controller is connected.
diff --git a/Engine/Plugins/Runtime/LeapMotionController/Source/LeapMotionController/Public/LeapMotionHandActor.h b/Engine/Plugins/Runtime/LeapMotionController/Source/LeapMotionController/Public/LeapMotionHandActor.h
--- a/Engine/Plugins/Runtime/LeapMotionController/Source/LeapMotionController/Public/LeapMotionHandActor.h
+++ b/Engine/Plugins/Runtime/LeapMotionController/Source/LeapMotionController/Public/LeapMotionHandActor.h
@@ -18,9 +18,7 @@ class ULeapMotionControllerComponent;
 UCLASS(ClassGroup = LeapMotion, BlueprintType, Blueprintable)
 class LEAPMOTIONCONTROLLER_API ALeapMotionHandActor : public AActor
 {
-	GENERATED_BODY()
-public:
-	ALeapMotionHandActor(const FObjectInitializer& ObjectInitializer = FObjectInitializer::Get());
+	GENERATED_UCLASS_BODY()
 
 	/** Scale of the hand, relatively to real-world size */
 	UPROPERTY(VisibleAnywhere, BlueprintReadOnly, Category = LeapMotion)
diff --git a/Engine/Plugins/Runtime/SoundMod/Source/SoundMod/Classes/SoundMod.h b/Engine/Plugins/Runtime/SoundMod/Source/SoundMod/Classes/SoundMod.h
--- a/Engine/Plugins/Runtime/SoundMod/Source/SoundMod/Classes/SoundMod.h
+++ b/Engine/Plugins/Runtime/SoundMod/Source/SoundMod/Classes/SoundMod.h
@@ -13,9 +13,7 @@ struct FActiveSound;
 UCLASS(hidecategories=Object, MinimalAPI, BlueprintType)
 class USoundMod : public USoundBase
 {
-	GENERATED_BODY()
-public:
-	SOUNDMOD_API USoundMod(const FObjectInitializer& ObjectInitializer = FObjectInitializer::Get());
+	GENERATED_UCLASS_BODY()
 
 	/** If set, when played directly (not through a sound cue) the nid will be played looping. */
 	UPROPERTY(EditAnywhere, Category=SoundMod)
diff --git a/Engine/Plugins/Runtime/SoundMod/Source/SoundMod/Classes/SoundModWave.h b/Engine/Plugins/Runtime/SoundMod/Source/SoundMod/Classes/SoundModWave.h
--- a/Engine/Plugins/Runtime/SoundMod/Source/SoundMod/Classes/SoundModWave.h
+++ b/Engine/Plugins/Runtime/SoundMod/Source/SoundMod/Classes/SoundModWave.h
@@ -11,9 +11,7 @@
 UCLASS()
 class USoundModWave : public USoundWaveStreaming
 {
-	GENERATED_BODY()
-public:
-	USoundModWave(const FObjectInitializer& ObjectInitializer = FObjectInitializer::Get());
+	GENERATED_UCLASS_BODY()
 
 	UPROPERTY()
 	class USoundMod* SoundMod;
@@ -24,4 +22,4 @@ public:
 	virtual int32 GeneratePCMData(uint8* PCMData, const int32 SamplesNeeded) override;
 	// End USoundWave interface.
 
-};
+};
\ No newline at end of file
diff --git a/Engine/Plugins/Runtime/SoundMod/Source/SoundMod/Public/SoundNodeModPlayer.h b/Engine/Plugins/Runtime/SoundMod/Source/SoundMod/Public/SoundNodeModPlayer.h
--- a/Engine/Plugins/Runtime/SoundMod/Source/SoundMod/Public/SoundNodeModPlayer.h
+++ b/Engine/Plugins/Runtime/SoundMod/Source/SoundMod/Public/SoundNodeModPlayer.h
@@ -11,9 +11,7 @@
 UCLASS(hidecategories=Object, editinlinenew, MinimalAPI, meta=( DisplayName="Mod Player" ))
 class USoundNodeModPlayer : public USoundNode
 {
-	GENERATED_BODY()
-public:
-	SOUNDMOD_API USoundNodeModPlayer(const FObjectInitializer& ObjectInitializer = FObjectInitializer::Get());
+	GENERATED_UCLASS_BODY()
 
 	UPROPERTY(EditAnywhere, Category=ModPlayer)
 	USoundMod* SoundMod;
diff --git a/Engine/Plugins/Runtime/TwitchLiveStreaming/Source/TwitchLiveStreaming/TwitchProjectSettings.h b/Engine/Plugins/Runtime/TwitchLiveStreaming/Source/TwitchLiveStreaming/TwitchProjectSettings.h
--- a/Engine/Plugins/Runtime/TwitchLiveStreaming/Source/TwitchLiveStreaming/TwitchProjectSettings.h
+++ b/Engine/Plugins/Runtime/TwitchLiveStreaming/Source/TwitchLiveStreaming/TwitchProjectSettings.h
@@ -6,9 +6,7 @@
 UCLASS(Config=Engine, DefaultConfig)
 class UTwitchProjectSettings : public UObject
 {
-	GENERATED_BODY()
-public:
-	UTwitchProjectSettings(const FObjectInitializer& ObjectInitializer = FObjectInitializer::Get());
+	GENERATED_UCLASS_BODY()
 
 public:
 
diff --git a/Engine/Plugins/Runtime/WebBrowserAsPlugin/Source/WebBrowserAsPlugin/Public/WebBrowser.h b/Engine/Plugins/Runtime/WebBrowserAsPlugin/Source/WebBrowserAsPlugin/Public/WebBrowser.h
--- a/Engine/Plugins/Runtime/WebBrowserAsPlugin/Source/WebBrowserAsPlugin/Public/WebBrowser.h
+++ b/Engine/Plugins/Runtime/WebBrowserAsPlugin/Source/WebBrowserAsPlugin/Public/WebBrowser.h
@@ -14,9 +14,7 @@
 UCLASS()
 class WEBBROWSERASPLUGIN_API UWebBrowser : public UWidget
 {
-	GENERATED_BODY()
-public:
-	UWebBrowser(const FObjectInitializer& ObjectInitializer = FObjectInitializer::Get());
+	GENERATED_UCLASS_BODY()
 
 public:
 
diff --git a/Engine/Plugins/ScriptPlugin/Source/ScriptEditorPlugin/Classes/ScriptFactory.h b/Engine/Plugins/ScriptPlugin/Source/ScriptEditorPlugin/Classes/ScriptFactory.h
--- a/Engine/Plugins/ScriptPlugin/Source/ScriptEditorPlugin/Classes/ScriptFactory.h
+++ b/Engine/Plugins/ScriptPlugin/Source/ScriptEditorPlugin/Classes/ScriptFactory.h
@@ -11,9 +11,7 @@
 UCLASS(collapsecategories, hidecategories = Object, EarlyAccessPreview)
 class SCRIPTEDITORPLUGIN_API UScriptFactory : public UFactory
 {
-	GENERATED_BODY()
-public:
-	UScriptFactory(const FObjectInitializer& ObjectInitializer = FObjectInitializer::Get());
+	GENERATED_UCLASS_BODY()
 
 public:
 
diff --git a/Engine/Plugins/ScriptPlugin/Source/ScriptPlugin/Classes/ScriptBlueprint.h b/Engine/Plugins/ScriptPlugin/Source/ScriptPlugin/Classes/ScriptBlueprint.h
--- a/Engine/Plugins/ScriptPlugin/Source/ScriptPlugin/Classes/ScriptBlueprint.h
+++ b/Engine/Plugins/ScriptPlugin/Source/ScriptPlugin/Classes/ScriptBlueprint.h
@@ -12,9 +12,7 @@
 UCLASS(BlueprintType, EarlyAccessPreview)
 class SCRIPTPLUGIN_API UScriptBlueprint : public UBlueprint
 {
-	GENERATED_BODY()
-public:
-	UScriptBlueprint(const FObjectInitializer& ObjectInitializer = FObjectInitializer::Get());
+	GENERATED_UCLASS_BODY()
 
 public:
 
@@ -56,4 +54,4 @@ public:
 	 */
 	void UpdateSourceCodeIfChanged();
 #endif
-};
+};
\ No newline at end of file
diff --git a/Engine/Plugins/ScriptPlugin/Source/ScriptPlugin/Classes/ScriptBlueprintGeneratedClass.h b/Engine/Plugins/ScriptPlugin/Source/ScriptPlugin/Classes/ScriptBlueprintGeneratedClass.h
--- a/Engine/Plugins/ScriptPlugin/Source/ScriptPlugin/Classes/ScriptBlueprintGeneratedClass.h
+++ b/Engine/Plugins/ScriptPlugin/Source/ScriptPlugin/Classes/ScriptBlueprintGeneratedClass.h
@@ -117,9 +117,7 @@ public:
 UCLASS(EarlyAccessPreview)
 class SCRIPTPLUGIN_API UScriptBlueprintGeneratedClass : public UBlueprintGeneratedClass
 {
-	GENERATED_BODY()
-public:
-	UScriptBlueprintGeneratedClass(const FObjectInitializer& ObjectInitializer = FObjectInitializer::Get());
+	GENERATED_UCLASS_BODY()
 
 	/** Generated script bytecode */
 	UPROPERTY()
@@ -161,4 +159,4 @@ public:
 		}
 		return ScriptClass;
 	}
-};
+};
\ No newline at end of file
diff --git a/Engine/Plugins/ScriptPlugin/Source/ScriptPlugin/Classes/ScriptContext.h b/Engine/Plugins/ScriptPlugin/Source/ScriptPlugin/Classes/ScriptContext.h
--- a/Engine/Plugins/ScriptPlugin/Source/ScriptPlugin/Classes/ScriptContext.h
+++ b/Engine/Plugins/ScriptPlugin/Source/ScriptPlugin/Classes/ScriptContext.h
@@ -10,9 +10,7 @@
 UCLASS(DefaultToInstanced, EarlyAccessPreview)
 class SCRIPTPLUGIN_API UScriptContext : public UObject
 {
-	GENERATED_BODY()
-public:
-	UScriptContext(const FObjectInitializer& ObjectInitializer = FObjectInitializer::Get());
+	GENERATED_UCLASS_BODY()
 
 	virtual void PostInitProperties() override;
 	virtual void BeginDestroy() override;
@@ -24,4 +22,4 @@ protected:
 
 	/** Script context (code) */
 	FScriptContextBase* Context;
-};
+};
\ No newline at end of file
diff --git a/Engine/Plugins/ScriptPlugin/Source/ScriptPlugin/Classes/ScriptContextComponent.h b/Engine/Plugins/ScriptPlugin/Source/ScriptPlugin/Classes/ScriptContextComponent.h
--- a/Engine/Plugins/ScriptPlugin/Source/ScriptPlugin/Classes/ScriptContextComponent.h
+++ b/Engine/Plugins/ScriptPlugin/Source/ScriptPlugin/Classes/ScriptContextComponent.h
@@ -10,9 +10,7 @@
 UCLASS(EarlyAccessPreview)
 class SCRIPTPLUGIN_API UScriptContextComponent : public UActorComponent
 {
-	GENERATED_BODY()
-public:
-	UScriptContextComponent(const FObjectInitializer& ObjectInitializer = FObjectInitializer::Get());
+	GENERATED_UCLASS_BODY()
 
 public:
 
diff --git a/Engine/Plugins/ScriptPlugin/Source/ScriptPlugin/Classes/ScriptPluginComponent.h b/Engine/Plugins/ScriptPlugin/Source/ScriptPlugin/Classes/ScriptPluginComponent.h
--- a/Engine/Plugins/ScriptPlugin/Source/ScriptPlugin/Classes/ScriptPluginComponent.h
+++ b/Engine/Plugins/ScriptPlugin/Source/ScriptPlugin/Classes/ScriptPluginComponent.h
@@ -10,9 +10,7 @@
 UCLASS(Blueprintable, hidecategories = (Object, ActorComponent), editinlinenew, meta = (BlueprintSpawnableComponent), ClassGroup = Script, Abstract, EarlyAccessPreview)
 class SCRIPTPLUGIN_API UScriptPluginComponent : public UActorComponent
 {
-	GENERATED_BODY()
-public:
-	UScriptPluginComponent(const FObjectInitializer& ObjectInitializer = FObjectInitializer::Get());
+	GENERATED_UCLASS_BODY()
 
 public:
 
diff --git a/Engine/Plugins/ScriptPlugin/Source/ScriptPlugin/Classes/ScriptTestActor.h b/Engine/Plugins/ScriptPlugin/Source/ScriptPlugin/Classes/ScriptTestActor.h
--- a/Engine/Plugins/ScriptPlugin/Source/ScriptPlugin/Classes/ScriptTestActor.h
+++ b/Engine/Plugins/ScriptPlugin/Source/ScriptPlugin/Classes/ScriptTestActor.h
@@ -8,9 +8,7 @@
 UCLASS(BlueprintType, MinimalAPI, EarlyAccessPreview, notplaceable)
 class AScriptTestActor : public AActor
 {
-	GENERATED_BODY()
-public:
-	SCRIPTPLUGIN_API AScriptTestActor(const FObjectInitializer& ObjectInitializer = FObjectInitializer::Get());
+	GENERATED_UCLASS_BODY()
 
 	UPROPERTY(Category = Test, EditAnywhere, BlueprintReadWrite)
 	FString TestString;
diff --git a/Engine/Plugins/Slate/SlateRemote/Source/SlateRemote/Private/Shared/SlateRemoteSettings.h b/Engine/Plugins/Slate/SlateRemote/Source/SlateRemote/Private/Shared/SlateRemoteSettings.h
--- a/Engine/Plugins/Slate/SlateRemote/Source/SlateRemote/Private/Shared/SlateRemoteSettings.h
+++ b/Engine/Plugins/Slate/SlateRemote/Source/SlateRemote/Private/Shared/SlateRemoteSettings.h
@@ -12,9 +12,7 @@ UCLASS(config=Engine)
 class USlateRemoteSettings
 	: public UObject
 {
-	GENERATED_BODY()
-public:
-	USlateRemoteSettings(const FObjectInitializer& ObjectInitializer = FObjectInitializer::Get());
+	GENERATED_UCLASS_BODY()
 
 public:
 
diff --git a/Engine/Source/Developer/AITestSuite/Classes/Actions/TestPawnAction_CallFunction.h b/Engine/Source/Developer/AITestSuite/Classes/Actions/TestPawnAction_CallFunction.h
--- a/Engine/Source/Developer/AITestSuite/Classes/Actions/TestPawnAction_CallFunction.h
+++ b/Engine/Source/Developer/AITestSuite/Classes/Actions/TestPawnAction_CallFunction.h
@@ -9,9 +9,7 @@
 UCLASS()
 class UTestPawnAction_CallFunction : public UTestPawnAction_Log
 {
-	GENERATED_BODY()
-public:
-	UTestPawnAction_CallFunction(const FObjectInitializer& ObjectInitializer = FObjectInitializer::Get());
+	GENERATED_UCLASS_BODY()
 
 	typedef void(*FFunctionToCall)(UPawnActionsComponent& ActionsComponent, UTestPawnAction_CallFunction& Caller, ETestPawnActionMessage::Type);
 	FFunctionToCall FunctionToCall;
@@ -23,4 +21,4 @@ public:
 	virtual bool Resume() override;
 	virtual void OnFinished(EPawnActionResult::Type WithResult) override;
 	virtual void OnChildFinished(UPawnAction& Action, EPawnActionResult::Type WithResult) override;
-};
+};
\ No newline at end of file
diff --git a/Engine/Source/Developer/AITestSuite/Classes/Actions/TestPawnAction_Log.h b/Engine/Source/Developer/AITestSuite/Classes/Actions/TestPawnAction_Log.h
--- a/Engine/Source/Developer/AITestSuite/Classes/Actions/TestPawnAction_Log.h
+++ b/Engine/Source/Developer/AITestSuite/Classes/Actions/TestPawnAction_Log.h
@@ -20,9 +20,7 @@ namespace ETestPawnActionMessage
 UCLASS()
 class UTestPawnAction_Log : public UPawnAction
 {
-	GENERATED_BODY()
-public:
-	UTestPawnAction_Log(const FObjectInitializer& ObjectInitializer = FObjectInitializer::Get());
+	GENERATED_UCLASS_BODY()
 
 	FTestLogger<int32>* Logger;
 
diff --git a/Engine/Source/Developer/AITestSuite/Classes/BehaviorTree/TestBTDecorator_CantExecute.h b/Engine/Source/Developer/AITestSuite/Classes/BehaviorTree/TestBTDecorator_CantExecute.h
--- a/Engine/Source/Developer/AITestSuite/Classes/BehaviorTree/TestBTDecorator_CantExecute.h
+++ b/Engine/Source/Developer/AITestSuite/Classes/BehaviorTree/TestBTDecorator_CantExecute.h
@@ -7,9 +7,7 @@
 UCLASS(meta=(HiddenNode))
 class UTestBTDecorator_CantExecute : public UBTDecorator
 {
-	GENERATED_BODY()
-public:
-	UTestBTDecorator_CantExecute(const FObjectInitializer& ObjectInitializer = FObjectInitializer::Get());
+	GENERATED_UCLASS_BODY()
 
 	virtual bool CalculateRawConditionValue(UBehaviorTreeComponent& OwnerComp, uint8* NodeMemory) const override;
 };
diff --git a/Engine/Source/Developer/AITestSuite/Classes/BehaviorTree/TestBTDecorator_DelayedAbort.h b/Engine/Source/Developer/AITestSuite/Classes/BehaviorTree/TestBTDecorator_DelayedAbort.h
--- a/Engine/Source/Developer/AITestSuite/Classes/BehaviorTree/TestBTDecorator_DelayedAbort.h
+++ b/Engine/Source/Developer/AITestSuite/Classes/BehaviorTree/TestBTDecorator_DelayedAbort.h
@@ -12,9 +12,7 @@ struct FBTDelayedAbortMemory
 UCLASS(meta = (HiddenNode))
 class UTestBTDecorator_DelayedAbort : public UBTDecorator
 {
-	GENERATED_BODY()
-public:
-	UTestBTDecorator_DelayedAbort(const FObjectInitializer& ObjectInitializer = FObjectInitializer::Get());
+	GENERATED_UCLASS_BODY()
 
 	UPROPERTY()
 	int32 DelayTicks;
diff --git a/Engine/Source/Developer/AITestSuite/Classes/BehaviorTree/TestBTTask_LatentWithFlags.h b/Engine/Source/Developer/AITestSuite/Classes/BehaviorTree/TestBTTask_LatentWithFlags.h
--- a/Engine/Source/Developer/AITestSuite/Classes/BehaviorTree/TestBTTask_LatentWithFlags.h
+++ b/Engine/Source/Developer/AITestSuite/Classes/BehaviorTree/TestBTTask_LatentWithFlags.h
@@ -15,9 +15,7 @@ struct FBTLatentTaskMemory
 UCLASS(meta = (HiddenNode))
 class UTestBTTask_LatentWithFlags : public UBTTaskNode
 {
-	GENERATED_BODY()
-public:
-	UTestBTTask_LatentWithFlags(const FObjectInitializer& ObjectInitializer = FObjectInitializer::Get());
+	GENERATED_UCLASS_BODY()
 
 	UPROPERTY()
 	int32 LogIndexExecuteStart;
diff --git a/Engine/Source/Developer/AITestSuite/Classes/BehaviorTree/TestBTTask_Log.h b/Engine/Source/Developer/AITestSuite/Classes/BehaviorTree/TestBTTask_Log.h
--- a/Engine/Source/Developer/AITestSuite/Classes/BehaviorTree/TestBTTask_Log.h
+++ b/Engine/Source/Developer/AITestSuite/Classes/BehaviorTree/TestBTTask_Log.h
@@ -12,9 +12,7 @@ struct FBTLogTaskMemory
 UCLASS(meta=(HiddenNode))
 class UTestBTTask_Log : public UBTTaskNode
 {
-	GENERATED_BODY()
-public:
-	UTestBTTask_Log(const FObjectInitializer& ObjectInitializer = FObjectInitializer::Get());
+	GENERATED_UCLASS_BODY()
 
 	UPROPERTY()
 	int32 LogIndex;
diff --git a/Engine/Source/Developer/AITestSuite/Classes/BehaviorTree/TestBTTask_SetFlag.h b/Engine/Source/Developer/AITestSuite/Classes/BehaviorTree/TestBTTask_SetFlag.h
--- a/Engine/Source/Developer/AITestSuite/Classes/BehaviorTree/TestBTTask_SetFlag.h
+++ b/Engine/Source/Developer/AITestSuite/Classes/BehaviorTree/TestBTTask_SetFlag.h
@@ -7,9 +7,7 @@
 UCLASS(meta=(HiddenNode))
 class UTestBTTask_SetFlag : public UBTTaskNode
 {
-	GENERATED_BODY()
-public:
-	UTestBTTask_SetFlag(const FObjectInitializer& ObjectInitializer = FObjectInitializer::Get());
+	GENERATED_UCLASS_BODY()
 
 	UPROPERTY()
 	FName KeyName;
diff --git a/Engine/Source/Developer/AITestSuite/Classes/MockAI_BT.h b/Engine/Source/Developer/AITestSuite/Classes/MockAI_BT.h
--- a/Engine/Source/Developer/AITestSuite/Classes/MockAI_BT.h
+++ b/Engine/Source/Developer/AITestSuite/Classes/MockAI_BT.h
@@ -7,9 +7,7 @@
 UCLASS()
 class UMockAI_BT : public UMockAI
 {
-	GENERATED_BODY()
-public:
-	UMockAI_BT(const FObjectInitializer& ObjectInitializer = FObjectInitializer::Get());
+	GENERATED_UCLASS_BODY()
 
 	UPROPERTY()
 	UBehaviorTreeComponent* BTComp;
diff --git a/Engine/Source/Developer/FriendsAndChat/Public/FriendsAndChatStyle.h b/Engine/Source/Developer/FriendsAndChat/Public/FriendsAndChatStyle.h
--- a/Engine/Source/Developer/FriendsAndChat/Public/FriendsAndChatStyle.h
+++ b/Engine/Source/Developer/FriendsAndChat/Public/FriendsAndChatStyle.h
@@ -21,7 +21,7 @@ USTRUCT()
 struct FRIENDSANDCHAT_API FFriendsAndChatStyle
 	: public FSlateWidgetStyle
 {
-	GENERATED_BODY()
+	GENERATED_USTRUCT_BODY()
 
 	// Default Constructor
 	FFriendsAndChatStyle() { }
diff --git a/Engine/Source/Developer/FunctionalTesting/Classes/FuncTestRenderingComponent.h b/Engine/Source/Developer/FunctionalTesting/Classes/FuncTestRenderingComponent.h
--- a/Engine/Source/Developer/FunctionalTesting/Classes/FuncTestRenderingComponent.h
+++ b/Engine/Source/Developer/FunctionalTesting/Classes/FuncTestRenderingComponent.h
@@ -7,9 +7,7 @@
 UCLASS(hidecategories=Object, editinlinenew)
 class UFuncTestRenderingComponent : public UPrimitiveComponent
 {
-	GENERATED_BODY()
-public:
-	UFuncTestRenderingComponent(const FObjectInitializer& ObjectInitializer = FObjectInitializer::Get());
+	GENERATED_UCLASS_BODY()
 
 	// Begin UPrimitiveComponent Interface
 	virtual FPrimitiveSceneProxy* CreateSceneProxy() override;
@@ -19,4 +17,4 @@ public:
 	// Begin USceneComponent Interface
 	virtual FBoxSphereBounds CalcBounds(const FTransform &LocalToWorld) const override;
 	// End USceneComponent Interface
-};
+};
\ No newline at end of file
diff --git a/Engine/Source/Developer/FunctionalTesting/Classes/FunctionalAITest.h b/Engine/Source/Developer/FunctionalTesting/Classes/FunctionalAITest.h
--- a/Engine/Source/Developer/FunctionalTesting/Classes/FunctionalAITest.h
+++ b/Engine/Source/Developer/FunctionalTesting/Classes/FunctionalAITest.h
@@ -14,7 +14,7 @@ DECLARE_DYNAMIC_MULTICAST_DELEGATE_TwoParams(FFunctionalTestAISpawned, AAIContro
 USTRUCT()
 struct FAITestSpawnInfo
 {
-	GENERATED_BODY()
+	GENERATED_USTRUCT_BODY()
 
 	/** Determines AI to be spawned */
 	UPROPERTY(EditAnywhere, BlueprintReadWrite, Category=AISpawn)
@@ -56,7 +56,7 @@ struct FAITestSpawnInfo
 USTRUCT()
 struct FPendingDelayedSpawn : public FAITestSpawnInfo
 {
-	GENERATED_BODY()
+	GENERATED_USTRUCT_BODY()
 
 	uint32 NumberToSpawnLeft;
 	float TimeToNextSpawn;
@@ -73,7 +73,7 @@ struct FPendingDelayedSpawn : public FAITestSpawnInfo
 USTRUCT()
 struct FAITestSpawnSet
 {
-	GENERATED_BODY()
+	GENERATED_USTRUCT_BODY()
 
 	/** what to spawn */
 	UPROPERTY(EditAnywhere, BlueprintReadWrite, Category=AISpawn)
@@ -97,9 +97,7 @@ struct FAITestSpawnSet
 UCLASS(Blueprintable, MinimalAPI)
 class AFunctionalAITest : public AFunctionalTest
 {
-	GENERATED_BODY()
-public:
-	FUNCTIONALTESTING_API AFunctionalAITest(const FObjectInitializer& ObjectInitializer = FObjectInitializer::Get());
+	GENERATED_UCLASS_BODY()
 	
 protected:
 	UPROPERTY(EditAnywhere, BlueprintReadOnly, Category=AITest)
@@ -151,4 +149,4 @@ protected:
 
 	void KillOffSpawnedPawns();
 	void ClearPendingDelayedSpawns();
-};
+};
\ No newline at end of file
diff --git a/Engine/Source/Developer/FunctionalTesting/Classes/FunctionalTest.h b/Engine/Source/Developer/FunctionalTesting/Classes/FunctionalTest.h
--- a/Engine/Source/Developer/FunctionalTesting/Classes/FunctionalTest.h
+++ b/Engine/Source/Developer/FunctionalTesting/Classes/FunctionalTest.h
@@ -26,9 +26,7 @@ DECLARE_DELEGATE_OneParam(FFunctionalTestDoneSignature, class AFunctionalTest*);
 UCLASS(Blueprintable, MinimalAPI)
 class AFunctionalTest : public AActor
 {
-	GENERATED_BODY()
-public:
-	FUNCTIONALTESTING_API AFunctionalTest(const FObjectInitializer& ObjectInitializer = FObjectInitializer::Get());
+	GENERATED_UCLASS_BODY()
 
 	static const uint32 DefaultTimeLimit = 60;	// seconds
 
diff --git a/Engine/Source/Developer/FunctionalTesting/Classes/FunctionalTestingManager.h b/Engine/Source/Developer/FunctionalTesting/Classes/FunctionalTestingManager.h
--- a/Engine/Source/Developer/FunctionalTesting/Classes/FunctionalTestingManager.h
+++ b/Engine/Source/Developer/FunctionalTesting/Classes/FunctionalTestingManager.h
@@ -17,9 +17,7 @@ namespace FFunctionalTesting
 UCLASS(BlueprintType, MinimalAPI)
 class UFunctionalTestingManager : public UBlueprintFunctionLibrary
 {
-	GENERATED_BODY()
-public:
-	FUNCTIONALTESTING_API UFunctionalTestingManager(const FObjectInitializer& ObjectInitializer = FObjectInitializer::Get());
+	GENERATED_UCLASS_BODY()
 	
 	UPROPERTY(Transient)
 	TArray<class AFunctionalTest*> TestsLeft;
diff --git a/Engine/Source/Developer/GameplayDebugger/Classes/GameplayDebuggerSettings.h b/Engine/Source/Developer/GameplayDebugger/Classes/GameplayDebuggerSettings.h
--- a/Engine/Source/Developer/GameplayDebugger/Classes/GameplayDebuggerSettings.h
+++ b/Engine/Source/Developer/GameplayDebugger/Classes/GameplayDebuggerSettings.h
@@ -13,7 +13,7 @@ GameplayDebuggerSettings.h: Declares the UGameplayDebuggerSettings class.
 USTRUCT()
 struct FGDTCustomViewNames
 {
-	GENERATED_BODY()
+	GENERATED_USTRUCT_BODY()
 
 	UPROPERTY(EditAnywhere, config, Category = "Settings")
 	FString GameView1;
@@ -34,9 +34,7 @@ struct FGDTCustomViewNames
 UCLASS(config = EditorUserSettings)
 class GAMEPLAYDEBUGGER_API UGameplayDebuggerSettings : public UObject
 {
-	GENERATED_BODY()
-public:
-	UGameplayDebuggerSettings(const FObjectInitializer& ObjectInitializer = FObjectInitializer::Get());
+	GENERATED_UCLASS_BODY()
 public:
 #if WITH_EDITOR
 	// Begin UObject Interface
diff --git a/Engine/Source/Developer/GameplayDebugger/Classes/GameplayDebuggingComponent.h b/Engine/Source/Developer/GameplayDebugger/Classes/GameplayDebuggingComponent.h
--- a/Engine/Source/Developer/GameplayDebugger/Classes/GameplayDebuggingComponent.h
+++ b/Engine/Source/Developer/GameplayDebugger/Classes/GameplayDebuggingComponent.h
@@ -36,9 +36,7 @@ DECLARE_MULTICAST_DELEGATE_TwoParams(FOnDebuggingTargetChanged, class AActor* /*
 UCLASS(config=Engine)
 class GAMEPLAYDEBUGGER_API UGameplayDebuggingComponent : public UPrimitiveComponent, public IEQSQueryResultSourceInterface
 {
-	GENERATED_BODY()
-public:
-	UGameplayDebuggingComponent(const FObjectInitializer& ObjectInitializer = FObjectInitializer::Get());
+	GENERATED_UCLASS_BODY()
 
 	friend class AGameplayDebuggingHUDComponent;
 
@@ -169,15 +167,11 @@ public:
 	UFUNCTION(exec)
 	void ServerReplicateData(uint32 InMessage, uint32 DataView);
 
-	UFUNCTION(reliable, server="ServerCollectNavmeshData_Implementation", WithValidation="ServerCollectNavmeshData_Validate")
+	UFUNCTION(reliable, server, WithValidation)
 	void ServerCollectNavmeshData(FVector_NetQuantize10 TargetLocation);
-	virtual void ServerCollectNavmeshData_Implementation(FVector_NetQuantize10 TargetLocation);
-	virtual bool ServerCollectNavmeshData_Validate(FVector_NetQuantize10 );
 
-	UFUNCTION(reliable, server="ServerDiscardNavmeshData_Implementation", WithValidation="ServerDiscardNavmeshData_Validate")
+	UFUNCTION(reliable, server, WithValidation)
 	void ServerDiscardNavmeshData();
-	virtual void ServerDiscardNavmeshData_Implementation();
-	virtual bool ServerDiscardNavmeshData_Validate();
 
 	void PrepareNavMeshData(struct FNavMeshSceneProxyData*) const;
 
@@ -200,10 +194,8 @@ public:
 	UPROPERTY(Replicated)
 	AActor* TargetActor;
 	
-	UFUNCTION(Reliable, Client="ClientEnableTargetSelection_Implementation", WithValidation="ClientEnableTargetSelection_Validate")
+	UFUNCTION(Reliable, Client, WithValidation)
 	void ClientEnableTargetSelection(bool bEnable);
-	virtual void ClientEnableTargetSelection_Implementation(bool bEnable);
-	virtual bool ClientEnableTargetSelection_Validate(bool );
 
 	void SetActorToDebug(AActor* Actor);
 	FORCEINLINE AActor* GetSelectedActor() const
diff --git a/Engine/Source/Developer/GameplayDebugger/Classes/GameplayDebuggingControllerComponent.h b/Engine/Source/Developer/GameplayDebugger/Classes/GameplayDebuggingControllerComponent.h
--- a/Engine/Source/Developer/GameplayDebugger/Classes/GameplayDebuggingControllerComponent.h
+++ b/Engine/Source/Developer/GameplayDebugger/Classes/GameplayDebuggingControllerComponent.h
@@ -15,9 +15,7 @@ class AGameplayDebuggingReplicator;
 UCLASS(config = Engine)
 class GAMEPLAYDEBUGGER_API UGameplayDebuggingControllerComponent : public UActorComponent
 {
-	GENERATED_BODY()
-public:
-	UGameplayDebuggingControllerComponent(const FObjectInitializer& ObjectInitializer = FObjectInitializer::Get());
+	GENERATED_UCLASS_BODY()
 
 #if WITH_HOT_RELOAD_CTORS
 	/** DO NOT USE. This constructor is for internal usage only for hot-reload purposes. */
diff --git a/Engine/Source/Developer/GameplayDebugger/Classes/GameplayDebuggingHUDComponent.h b/Engine/Source/Developer/GameplayDebugger/Classes/GameplayDebuggingHUDComponent.h
--- a/Engine/Source/Developer/GameplayDebugger/Classes/GameplayDebuggingHUDComponent.h
+++ b/Engine/Source/Developer/GameplayDebugger/Classes/GameplayDebuggingHUDComponent.h
@@ -24,9 +24,7 @@ struct GAMEPLAYDEBUGGER_API FDebugCategoryView
 UCLASS(config = Engine, notplaceable)
 class GAMEPLAYDEBUGGER_API AGameplayDebuggingHUDComponent : public AActor
 {
-	GENERATED_BODY()
-public:
-	AGameplayDebuggingHUDComponent(const FObjectInitializer& ObjectInitializer = FObjectInitializer::Get());
+	GENERATED_UCLASS_BODY()
 
 	UPROPERTY(config)
 	float MenuStartX;
diff --git a/Engine/Source/Developer/GameplayDebugger/Classes/GameplayDebuggingReplicator.h b/Engine/Source/Developer/GameplayDebugger/Classes/GameplayDebuggingReplicator.h
--- a/Engine/Source/Developer/GameplayDebugger/Classes/GameplayDebuggingReplicator.h
+++ b/Engine/Source/Developer/GameplayDebugger/Classes/GameplayDebuggingReplicator.h
@@ -15,9 +15,7 @@ DECLARE_MULTICAST_DELEGATE_OneParam(FOnSelectionChanged, class AActor*);
 UCLASS(config = Engine, NotBlueprintable, Transient, hidecategories = Actor, notplaceable)
 class GAMEPLAYDEBUGGER_API AGameplayDebuggingReplicator : public AActor
 {
-	GENERATED_BODY()
-public:
-	AGameplayDebuggingReplicator(const FObjectInitializer& ObjectInitializer = FObjectInitializer::Get());
+	GENERATED_UCLASS_BODY()
 
 	UPROPERTY(config)
 	FString DebugComponentClassName;
@@ -88,24 +86,17 @@ public:
 	UPROPERTY()
 	class UTexture2D* DefaultTexture_Green;
 
-	UFUNCTION(reliable, server="ServerReplicateMessage_Implementation", WithValidation="ServerReplicateMessage_Validate")
+	UFUNCTION(reliable, server, WithValidation)
 	void ServerReplicateMessage(class AActor* Actor, uint32 InMessage, uint32 DataView = 0);
-	virtual void ServerReplicateMessage_Implementation(AActor* Actor, uint32 InMessage, uint32 DataView);
-	virtual bool ServerReplicateMessage_Validate(AActor* , uint32 , uint32 );
 
-	UFUNCTION(reliable, client="ClientReplicateMessage_Implementation", WithValidation="ClientReplicateMessage_Validate")
+	UFUNCTION(reliable, client, WithValidation)
 	void ClientReplicateMessage(class AActor* Actor, uint32 InMessage, uint32 DataView = 0);
-	virtual void ClientReplicateMessage_Implementation(AActor* Actor, uint32 InMessage, uint32 DataView);
-	virtual bool ClientReplicateMessage_Validate(AActor* , uint32 , uint32 );
 
-	UFUNCTION(Reliable, Client="ClientAutoActivate_Implementation")
+	UFUNCTION(Reliable, Client)
 	void ClientAutoActivate();
-	virtual void ClientAutoActivate_Implementation();
 
-	UFUNCTION(Reliable, Client="ClientEnableTargetSelection_Implementation", WithValidation="ClientEnableTargetSelection_Validate")
+	UFUNCTION(Reliable, Client, WithValidation)
 	void ClientEnableTargetSelection(bool bEnable, APlayerController* Context);
-	virtual void ClientEnableTargetSelection_Implementation(bool bEnable, APlayerController* Context);
-	virtual bool ClientEnableTargetSelection_Validate(bool , APlayerController* );
 
 #if WITH_EDITOR
 	void PostEditChangeProperty(FPropertyChangedEvent& PropertyChangedEvent) override;
@@ -141,10 +132,8 @@ public:
 
 	FORCEINLINE AActor* GetSelectedActorToDebug() { return LastSelectedActorToDebug; }
 	
-	UFUNCTION(reliable, server="ServerSetActorToDebug_Implementation", WithValidation="ServerSetActorToDebug_Validate")
+	UFUNCTION(reliable, server, WithValidation)
 	void ServerSetActorToDebug(AActor* InActor);
-	virtual void ServerSetActorToDebug_Implementation(AActor* InActor);
-	virtual bool ServerSetActorToDebug_Validate(AActor* );
 
 	/**
 	 * Iterates through the pawn list to find the next pawn of the specified type to debug
diff --git a/Engine/Source/Developer/HTML5/HTML5PlatformEditor/Classes/HTML5TargetSettings.h b/Engine/Source/Developer/HTML5/HTML5PlatformEditor/Classes/HTML5TargetSettings.h
--- a/Engine/Source/Developer/HTML5/HTML5PlatformEditor/Classes/HTML5TargetSettings.h
+++ b/Engine/Source/Developer/HTML5/HTML5PlatformEditor/Classes/HTML5TargetSettings.h
@@ -16,11 +16,9 @@ UCLASS(config=Engine, defaultconfig)
 class HTML5PLATFORMEDITOR_API UHTML5TargetSettings
 	: public UObject
 {
-	GENERATED_BODY()
-public:
-	UHTML5TargetSettings(const FObjectInitializer& ObjectInitializer = FObjectInitializer::Get());
 public:
 
+	GENERATED_UCLASS_BODY()
 
  	/**
  	 * Setting to control HTML5 Heap size (in Development)
diff --git a/Engine/Source/Developer/HTML5/HTML5PlatformEditor/Private/HTML5SDKSettings.h b/Engine/Source/Developer/HTML5/HTML5PlatformEditor/Private/HTML5SDKSettings.h
--- a/Engine/Source/Developer/HTML5/HTML5PlatformEditor/Private/HTML5SDKSettings.h
+++ b/Engine/Source/Developer/HTML5/HTML5PlatformEditor/Private/HTML5SDKSettings.h
@@ -14,7 +14,7 @@ class IHTML5TargetPlatformModule;
 USTRUCT()
 struct FHTML5DeviceMapping
 {
-	GENERATED_BODY()
+	GENERATED_USTRUCT_BODY()
 
 	UPROPERTY(EditAnywhere, Category = HTML5_Devices, Meta = (DisplayName = "Browser name"))
 	FString DeviceName;
@@ -26,7 +26,7 @@ struct FHTML5DeviceMapping
 USTRUCT()
 struct FHTML5SDKPath
 {
-	GENERATED_BODY()
+	GENERATED_USTRUCT_BODY()
 
 	UPROPERTY(EditAnywhere, Category = SDK_Path)
 	FString SDKPath;
@@ -42,10 +42,8 @@ struct FHTML5SDKPath
 UCLASS(config=Engine, globaluserconfig)
 class HTML5PLATFORMEDITOR_API UHTML5SDKSettings : public UObject
 {
-	GENERATED_BODY()
-public:
-	UHTML5SDKSettings(const FObjectInitializer& ObjectInitializer = FObjectInitializer::Get());
 public:
+	GENERATED_UCLASS_BODY()
 
 	UPROPERTY(GlobalConfig, EditAnywhere, Category = HTML5_SDK_Paths, Meta = (DisplayName = "Location of Python exe (This can be blank if python is on the PATH)"))
 	FFilePath Python;
diff --git a/Engine/Source/Developer/InternationalizationSettings/Classes/InternationalizationSettingsModel.h b/Engine/Source/Developer/InternationalizationSettings/Classes/InternationalizationSettingsModel.h
--- a/Engine/Source/Developer/InternationalizationSettings/Classes/InternationalizationSettingsModel.h
+++ b/Engine/Source/Developer/InternationalizationSettings/Classes/InternationalizationSettingsModel.h
@@ -12,9 +12,7 @@ UCLASS()
 class INTERNATIONALIZATIONSETTINGS_API UInternationalizationSettingsModel
 	:	public UObject
 {
-	GENERATED_BODY()
-public:
-	UInternationalizationSettingsModel(const FObjectInitializer& ObjectInitializer = FObjectInitializer::Get());
+	GENERATED_UCLASS_BODY()
 
 public:
 
diff --git a/Engine/Source/Developer/Linux/LinuxTargetPlatform/Classes/LinuxTargetSettings.h b/Engine/Source/Developer/Linux/LinuxTargetPlatform/Classes/LinuxTargetSettings.h
--- a/Engine/Source/Developer/Linux/LinuxTargetPlatform/Classes/LinuxTargetSettings.h
+++ b/Engine/Source/Developer/Linux/LinuxTargetPlatform/Classes/LinuxTargetSettings.h
@@ -16,10 +16,9 @@ UCLASS(config=Engine, defaultconfig)
 class LINUXTARGETPLATFORM_API ULinuxTargetSettings
 	: public UObject
 {
-	GENERATED_BODY()
-
 public:
-	ULinuxTargetSettings(const FObjectInitializer& ObjectInitializer = FObjectInitializer::Get());
+
+	GENERATED_UCLASS_BODY()
 
 	/**
 	 * The collection of RHI's we want to support on this platform.
diff --git a/Engine/Source/Developer/LogVisualizer/Private/VisualLoggerCameraController.h b/Engine/Source/Developer/LogVisualizer/Private/VisualLoggerCameraController.h
--- a/Engine/Source/Developer/LogVisualizer/Private/VisualLoggerCameraController.h
+++ b/Engine/Source/Developer/LogVisualizer/Private/VisualLoggerCameraController.h
@@ -8,13 +8,11 @@
 UCLASS(config=Input, hidedropdown)
 class AVisualLoggerCameraController : public ADebugCameraController
 {
-	GENERATED_BODY()
-public:
-	AVisualLoggerCameraController(const FObjectInitializer& ObjectInitializer = FObjectInitializer::Get());
 public:
 	DECLARE_DELEGATE_OneParam(FActorSelectedDelegate, class AActor*);
 	DECLARE_DELEGATE_OneParam(FLogEntryIterationDelegate, int32);
 
+	GENERATED_UCLASS_BODY()
 
 	UPROPERTY()
 	class AActor* PickedActor;
diff --git a/Engine/Source/Developer/LogVisualizer/Private/VisualLoggerHUD.h b/Engine/Source/Developer/LogVisualizer/Private/VisualLoggerHUD.h
--- a/Engine/Source/Developer/LogVisualizer/Private/VisualLoggerHUD.h
+++ b/Engine/Source/Developer/LogVisualizer/Private/VisualLoggerHUD.h
@@ -7,9 +7,7 @@
 UCLASS()
 class AVisualLoggerHUD : public ADebugCameraHUD
 {
-	GENERATED_BODY()
-public:
-	AVisualLoggerHUD(const FObjectInitializer& ObjectInitializer = FObjectInitializer::Get());
+	GENERATED_UCLASS_BODY()
 
 	FFontRenderInfo TextRenderInfo;
 
diff --git a/Engine/Source/Developer/LogVisualizer/Private/VisualLoggerRenderingActor.h b/Engine/Source/Developer/LogVisualizer/Private/VisualLoggerRenderingActor.h
--- a/Engine/Source/Developer/LogVisualizer/Private/VisualLoggerRenderingActor.h
+++ b/Engine/Source/Developer/LogVisualizer/Private/VisualLoggerRenderingActor.h
@@ -13,10 +13,8 @@ DECLARE_MULTICAST_DELEGATE_OneParam(FOnSelectionChanged, class AActor*);
 UCLASS(config = Engine, NotBlueprintable, Transient)
 class LOGVISUALIZER_API AVisualLoggerRenderingActor : public AActor
 {
-	GENERATED_BODY()
-public:
-	AVisualLoggerRenderingActor(const FObjectInitializer& ObjectInitializer = FObjectInitializer::Get());
 public:
+	GENERATED_UCLASS_BODY()
 
 	void OnItemSelectionChanged(const FVisualLogDevice::FVisualLogEntryItem& EntryItem);
 	void ObjectSelectionChanged(TSharedPtr<class STimeline> TimeLine);
diff --git a/Engine/Source/Developer/LogVisualizer/Private/VisualLoggerRenderingComponent.h b/Engine/Source/Developer/LogVisualizer/Private/VisualLoggerRenderingComponent.h
--- a/Engine/Source/Developer/LogVisualizer/Private/VisualLoggerRenderingComponent.h
+++ b/Engine/Source/Developer/LogVisualizer/Private/VisualLoggerRenderingComponent.h
@@ -10,10 +10,8 @@
 UCLASS()
 class UVisualLoggerRenderingComponent : public UPrimitiveComponent
 {
-	GENERATED_BODY()
-public:
-	UVisualLoggerRenderingComponent(const FObjectInitializer& ObjectInitializer = FObjectInitializer::Get());
 public:
+	GENERATED_UCLASS_BODY()
 
 	virtual FPrimitiveSceneProxy* CreateSceneProxy() override;
 	virtual FBoxSphereBounds CalcBounds(const FTransform &LocalToWorld) const override;
diff --git a/Engine/Source/Developer/LogVisualizer/Public/LogVisualizerSessionSettings.h b/Engine/Source/Developer/LogVisualizer/Public/LogVisualizerSessionSettings.h
--- a/Engine/Source/Developer/LogVisualizer/Public/LogVisualizerSessionSettings.h
+++ b/Engine/Source/Developer/LogVisualizer/Public/LogVisualizerSessionSettings.h
@@ -11,9 +11,7 @@ GameplayDebuggerSettings.h: Declares the UGameplayDebuggerSettings class.
 UCLASS()
 class LOGVISUALIZER_API ULogVisualizerSessionSettings : public UObject
 {
-	GENERATED_BODY()
-public:
-	ULogVisualizerSessionSettings(const FObjectInitializer& ObjectInitializer = FObjectInitializer::Get());
+	GENERATED_UCLASS_BODY()
 public:
 	/**Whether to show trivial logs, i.e. the ones with only one entry.*/
 	UPROPERTY(EditAnywhere, Category = "VisualLogger")
diff --git a/Engine/Source/Developer/LogVisualizer/Public/LogVisualizerSettings.h b/Engine/Source/Developer/LogVisualizer/Public/LogVisualizerSettings.h
--- a/Engine/Source/Developer/LogVisualizer/Public/LogVisualizerSettings.h
+++ b/Engine/Source/Developer/LogVisualizer/Public/LogVisualizerSettings.h
@@ -11,7 +11,7 @@ GameplayDebuggerSettings.h: Declares the UGameplayDebuggerSettings class.
 USTRUCT()
 struct FCategoryFilter
 {
-	GENERATED_BODY()
+	GENERATED_USTRUCT_BODY()
 
 	UPROPERTY(config)
 	FString CategoryName;
@@ -26,7 +26,7 @@ struct FCategoryFilter
 USTRUCT()
 struct FVisualLoggerFilters
 {
-	GENERATED_BODY()
+	GENERATED_USTRUCT_BODY()
 
 	UPROPERTY(config)
 	FString SearchBoxFilter;
@@ -46,9 +46,7 @@ struct FCategoryFiltersManager;
 UCLASS(config = EditorUserSettings)
 class LOGVISUALIZER_API ULogVisualizerSettings : public UObject
 {
-	GENERATED_BODY()
-public:
-	ULogVisualizerSettings(const FObjectInitializer& ObjectInitializer = FObjectInitializer::Get());
+	GENERATED_UCLASS_BODY()
 	friend struct FCategoryFiltersManager;
 
 public:
diff --git a/Engine/Source/Developer/Mac/MacTargetPlatform/Classes/MacTargetSettings.h b/Engine/Source/Developer/Mac/MacTargetPlatform/Classes/MacTargetSettings.h
--- a/Engine/Source/Developer/Mac/MacTargetPlatform/Classes/MacTargetSettings.h
+++ b/Engine/Source/Developer/Mac/MacTargetPlatform/Classes/MacTargetSettings.h
@@ -16,7 +16,8 @@ UCLASS()
 class MACTARGETPLATFORM_API UMacTargetSettings
 	: public UObject
 {
-	GENERATED_BODY()
 public:
-	UMacTargetSettings(const FObjectInitializer& ObjectInitializer = FObjectInitializer::Get());
+
+	GENERATED_UCLASS_BODY()
+
 };
diff --git a/Engine/Source/Developer/MovieSceneTools/Classes/K2Node_PlayMovieScene.h b/Engine/Source/Developer/MovieSceneTools/Classes/K2Node_PlayMovieScene.h
--- a/Engine/Source/Developer/MovieSceneTools/Classes/K2Node_PlayMovieScene.h
+++ b/Engine/Source/Developer/MovieSceneTools/Classes/K2Node_PlayMovieScene.h
@@ -14,7 +14,7 @@ DECLARE_EVENT( UK2Node_PlayMovieScene, FOnMovieSceneBindingsChanged );
 USTRUCT()
 struct FEdGraphSchemaAction_K2AddPlayMovieScene : public FEdGraphSchemaAction_K2NewNode
 {
-	GENERATED_BODY()
+	GENERATED_USTRUCT_BODY()
 
 	FEdGraphSchemaAction_K2AddPlayMovieScene()
 		: FEdGraphSchemaAction_K2NewNode()
@@ -36,9 +36,7 @@ struct FEdGraphSchemaAction_K2AddPlayMovieScene : public FEdGraphSchemaAction_K2
 UCLASS( MinimalAPI )
 class UK2Node_PlayMovieScene : public UK2Node
 {
-	GENERATED_BODY()
-public:
-	MOVIESCENETOOLS_API UK2Node_PlayMovieScene(const FObjectInitializer& ObjectInitializer = FObjectInitializer::Get());
+	GENERATED_UCLASS_BODY()
 
 #if WITH_EDITOR
 	// UEdGraphNode interface.
diff --git a/Engine/Source/Developer/MovieSceneTools/Classes/MovieSceneFactory.h b/Engine/Source/Developer/MovieSceneTools/Classes/MovieSceneFactory.h
--- a/Engine/Source/Developer/MovieSceneTools/Classes/MovieSceneFactory.h
+++ b/Engine/Source/Developer/MovieSceneTools/Classes/MovieSceneFactory.h
@@ -7,9 +7,7 @@
 UCLASS()
 class UMovieSceneFactory : public UFactory
 {
-	GENERATED_BODY()
-public:
-	UMovieSceneFactory(const FObjectInitializer& ObjectInitializer = FObjectInitializer::Get());
+	GENERATED_UCLASS_BODY()
 
 	// Begin UFactory Interface
 	virtual UObject* FactoryCreateNew(UClass* Class,UObject* InParent,FName Name,EObjectFlags Flags,UObject* Context,FFeedbackContext* Warn) override;
diff --git a/Engine/Source/Developer/ProfilerMessages/Classes/ProfilerServiceMessages.h b/Engine/Source/Developer/ProfilerMessages/Classes/ProfilerServiceMessages.h
--- a/Engine/Source/Developer/ProfilerMessages/Classes/ProfilerServiceMessages.h
+++ b/Engine/Source/Developer/ProfilerMessages/Classes/ProfilerServiceMessages.h
@@ -10,7 +10,7 @@
 USTRUCT()
 struct FProfilerServiceAuthorize
 {
-	GENERATED_BODY()
+	GENERATED_USTRUCT_BODY()
 
 	/**
 	 */
@@ -53,7 +53,7 @@ struct TStructOpsTypeTraits<FProfilerServiceAuthorize> : public TStructOpsTypeTr
 USTRUCT()
 struct FProfilerServiceAuthorize2
 {
-	GENERATED_BODY()
+	GENERATED_USTRUCT_BODY()
 
 	/**
 	 */
@@ -102,7 +102,7 @@ struct TStructOpsTypeTraits<FProfilerServiceAuthorize2> : public TStructOpsTypeT
 USTRUCT()
 struct FProfilerServiceData
 {
-	GENERATED_BODY()
+	GENERATED_USTRUCT_BODY()
 
 	/**
 	 */
@@ -145,7 +145,7 @@ struct TStructOpsTypeTraits<FProfilerServiceData> : public TStructOpsTypeTraitsB
 USTRUCT()
 struct FProfilerServiceData2
 {
-	GENERATED_BODY()
+	GENERATED_USTRUCT_BODY()
 
 	/**
 	 */
@@ -192,7 +192,7 @@ struct TStructOpsTypeTraits<FProfilerServiceData2> : public TStructOpsTypeTraits
 USTRUCT()
 struct FProfilerServicePreviewAck
 {
-	GENERATED_BODY()
+	GENERATED_USTRUCT_BODY()
 
 	/**
 	 */
@@ -233,7 +233,7 @@ struct TStructOpsTypeTraits<FProfilerServicePreviewAck> : public TStructOpsTypeT
 USTRUCT()
 struct FProfilerServiceMetaData
 {
-	GENERATED_BODY()
+	GENERATED_USTRUCT_BODY()
 
 	/**
 	 */
@@ -277,7 +277,7 @@ struct TStructOpsTypeTraits<FProfilerServiceMetaData> : public TStructOpsTypeTra
 USTRUCT()
 struct FProfilerServiceFileChunk
 {
-	GENERATED_BODY()
+	GENERATED_USTRUCT_BODY()
 
 	/** The ID of the instance where this message should be sent. */
 	UPROPERTY()
@@ -345,7 +345,7 @@ struct TStructOpsTypeTraits<FProfilerServiceFileChunk> : public TStructOpsTypeTr
 USTRUCT()
 struct FProfilerServicePing
 {
-	GENERATED_BODY()
+	GENERATED_USTRUCT_BODY()
 };
 
 template<>
@@ -363,7 +363,7 @@ struct TStructOpsTypeTraits<FProfilerServicePing> : public TStructOpsTypeTraitsB
 USTRUCT()
 struct FProfilerServicePong
 {
-	GENERATED_BODY()
+	GENERATED_USTRUCT_BODY()
 
 };
 
@@ -382,7 +382,7 @@ struct TStructOpsTypeTraits<FProfilerServicePong> : public TStructOpsTypeTraitsB
 USTRUCT()
 struct FProfilerServiceSubscribe
 {
-	GENERATED_BODY()
+	GENERATED_USTRUCT_BODY()
 
 	/**
 	 */
@@ -423,7 +423,7 @@ struct TStructOpsTypeTraits<FProfilerServiceSubscribe> : public TStructOpsTypeTr
 USTRUCT()
 struct FProfilerServiceUnsubscribe
 {
-	GENERATED_BODY()
+	GENERATED_USTRUCT_BODY()
 
 	/**
 	 */
@@ -465,7 +465,7 @@ struct TStructOpsTypeTraits<FProfilerServiceUnsubscribe> : public TStructOpsType
 USTRUCT()
 struct FProfilerServiceCapture
 {
-	GENERATED_BODY()
+	GENERATED_USTRUCT_BODY()
 
 	/**
 	 * The data capture state that should be set.
@@ -501,7 +501,7 @@ struct TStructOpsTypeTraits<FProfilerServiceCapture> : public TStructOpsTypeTrai
 USTRUCT()
 struct FProfilerServicePreview
 {
-	GENERATED_BODY()
+	GENERATED_USTRUCT_BODY()
 
 	/**
 	 * The data preview state that should be set.
@@ -537,7 +537,7 @@ struct TStructOpsTypeTraits<FProfilerServicePreview> : public TStructOpsTypeTrai
 USTRUCT()
 struct FProfilerServiceRequest
 {
-	GENERATED_BODY()
+	GENERATED_USTRUCT_BODY()
 
 	/** Request @see EProfilerRequestType. */
 	UPROPERTY()
diff --git a/Engine/Source/Developer/SourceCodeAccess/Private/SourceCodeAccessSettings.h b/Engine/Source/Developer/SourceCodeAccess/Private/SourceCodeAccessSettings.h
--- a/Engine/Source/Developer/SourceCodeAccess/Private/SourceCodeAccessSettings.h
+++ b/Engine/Source/Developer/SourceCodeAccess/Private/SourceCodeAccessSettings.h
@@ -7,11 +7,9 @@
 UCLASS(config=EditorUserSettings)
 class USourceCodeAccessSettings : public UObject
 {
-	GENERATED_BODY()
-public:
-	USourceCodeAccessSettings(const FObjectInitializer& ObjectInitializer = FObjectInitializer::Get());
+	GENERATED_UCLASS_BODY()
 
 	/** The source code editor we prefer to use. */
 	UPROPERTY(Config, EditAnywhere, Category="Source Code Editor", meta=(DisplayName="Source Code Editor"))
 	FString PreferredAccessor;
-};
+};
\ No newline at end of file
diff --git a/Engine/Source/Developer/TargetDeviceServices/Classes/TargetDeviceServiceMessages.h b/Engine/Source/Developer/TargetDeviceServices/Classes/TargetDeviceServiceMessages.h
--- a/Engine/Source/Developer/TargetDeviceServices/Classes/TargetDeviceServiceMessages.h
+++ b/Engine/Source/Developer/TargetDeviceServices/Classes/TargetDeviceServiceMessages.h
@@ -16,7 +16,7 @@
 USTRUCT()
 struct FTargetDeviceServiceDeployCommit
 {
-	GENERATED_BODY()
+	GENERATED_USTRUCT_BODY()
 
 	/** Holds the variant identifier of the target device to use. */
 	UPROPERTY()
@@ -57,7 +57,7 @@ struct TStructOpsTypeTraits<FTargetDeviceServiceDeployCommit> : public TStructOp
 USTRUCT()
 struct FTargetDeviceServiceDeployFile
 {
-	GENERATED_BODY()
+	GENERATED_USTRUCT_BODY()
 
 	/** Holds the name and path of the file as it will be stored on the target device. */
 	UPROPERTY()
@@ -96,7 +96,7 @@ struct TStructOpsTypeTraits<FTargetDeviceServiceDeployFile> : public TStructOpsT
 USTRUCT()
 struct FTargetDeviceServiceDeployFinished
 {
-	GENERATED_BODY()
+	GENERATED_USTRUCT_BODY()
 
 	/** Holds the variant identifier of the target device to use. */
 	UPROPERTY()
@@ -157,7 +157,7 @@ struct TStructOpsTypeTraits<FTargetDeviceServiceDeployFinished> : public TStruct
 USTRUCT()
 struct FTargetDeviceServiceLaunchApp
 {
-	GENERATED_BODY()
+	GENERATED_USTRUCT_BODY()
 
 	/** Holds the variant identifier of the target device to use. */
 	UPROPERTY()
@@ -213,7 +213,7 @@ struct TStructOpsTypeTraits<FTargetDeviceServiceLaunchApp> : public TStructOpsTy
 USTRUCT()
 struct FTargetDeviceServiceLaunchFinished
 {
-	GENERATED_BODY()
+	GENERATED_USTRUCT_BODY()
 
 	/**
 	 * Holds the identifier of the launched application.
@@ -268,7 +268,7 @@ struct TStructOpsTypeTraits<FTargetDeviceServiceLaunchFinished> : public TStruct
 USTRUCT()
 struct FTargetDeviceClaimDenied
 {
-	GENERATED_BODY()
+	GENERATED_USTRUCT_BODY()
 
 	/** Holds the identifier of the device that is already claimed. */
 	UPROPERTY()
@@ -313,7 +313,7 @@ struct TStructOpsTypeTraits<FTargetDeviceClaimDenied> : public TStructOpsTypeTra
 USTRUCT()
 struct FTargetDeviceClaimed
 {
-	GENERATED_BODY()
+	GENERATED_USTRUCT_BODY()
 
 	/** Holds the identifier of the device that is being claimed. */
 	UPROPERTY()
@@ -357,7 +357,7 @@ struct TStructOpsTypeTraits<FTargetDeviceClaimed> : public TStructOpsTypeTraitsB
 USTRUCT()
 struct FTargetDeviceUnclaimed
 {
-	GENERATED_BODY()
+	GENERATED_USTRUCT_BODY()
 
 	/** Holds the identifier of the device that is no longer claimed. */
 	UPROPERTY()
@@ -402,7 +402,7 @@ struct TStructOpsTypeTraits<FTargetDeviceUnclaimed> : public TStructOpsTypeTrait
 USTRUCT()
 struct FTargetDeviceServicePing
 {
-	GENERATED_BODY()
+	GENERATED_USTRUCT_BODY()
 
 	/** Holds the name of the user who generated the ping. */
 	UPROPERTY()
@@ -434,7 +434,7 @@ struct TStructOpsTypeTraits<FTargetDeviceServicePing> : public TStructOpsTypeTra
 USTRUCT()
 struct FTargetDeviceVariant
 {
-	GENERATED_BODY()
+	GENERATED_USTRUCT_BODY()
 
 	UPROPERTY()
 	FString DeviceID;
@@ -471,7 +471,7 @@ struct TStructOpsTypeTraits<FTargetDeviceVariant> : public TStructOpsTypeTraitsB
 USTRUCT()
 struct FTargetDeviceServicePong
 {
-	GENERATED_BODY()
+	GENERATED_USTRUCT_BODY()
 
 	/** Holds a flag indicating whether the device is currently connected. */
 	UPROPERTY()
@@ -561,7 +561,7 @@ struct TStructOpsTypeTraits<FTargetDeviceServicePong> : public TStructOpsTypeTra
 USTRUCT()
 struct FTargetDeviceServicePowerOff
 {
-	GENERATED_BODY()
+	GENERATED_USTRUCT_BODY()
 
 	/**
 	 * Holds a flag indicating whether the power-off should be enforced.
@@ -602,7 +602,7 @@ struct TStructOpsTypeTraits<FTargetDeviceServicePowerOff> : public TStructOpsTyp
 USTRUCT()
 struct FTargetDeviceServicePowerOn
 {
-	GENERATED_BODY()
+	GENERATED_USTRUCT_BODY()
 
 	/** Holds the name of the user that wishes to power on the device. */
 	UPROPERTY()
@@ -634,7 +634,7 @@ struct TStructOpsTypeTraits<FTargetDeviceServicePowerOn> : public TStructOpsType
 USTRUCT()
 struct FTargetDeviceServiceReboot
 {
-	GENERATED_BODY()
+	GENERATED_USTRUCT_BODY()
 
 	/** Holds the name of the user that wishes to reboot the device. */
 	UPROPERTY()
@@ -670,7 +670,7 @@ struct TStructOpsTypeTraits<FTargetDeviceServiceReboot> : public TStructOpsTypeT
 USTRUCT()
 struct FTargetDeviceServiceRunExecutable
 {
-	GENERATED_BODY()
+	GENERATED_USTRUCT_BODY()
 
 	/** Holds the variant identifier of the target device to use for execution. */
 	UPROPERTY()
@@ -714,7 +714,7 @@ struct TStructOpsTypeTraits<FTargetDeviceServiceRunExecutable> : public TStructO
 USTRUCT()
 struct FTargetDeviceServiceRunFinished
 {
-	GENERATED_BODY()
+	GENERATED_USTRUCT_BODY()
 
 	/** Holds the variant identifier of the target device to use. */
 	UPROPERTY()
diff --git a/Engine/Source/Developer/WidgetCarousel/Public/WidgetCarouselStyle.h b/Engine/Source/Developer/WidgetCarousel/Public/WidgetCarouselStyle.h
--- a/Engine/Source/Developer/WidgetCarousel/Public/WidgetCarouselStyle.h
+++ b/Engine/Source/Developer/WidgetCarousel/Public/WidgetCarouselStyle.h
@@ -8,7 +8,7 @@
 USTRUCT()
 struct WIDGETCAROUSEL_API FWidgetCarouselNavigationButtonStyle : public FSlateWidgetStyle
 {
-	GENERATED_BODY()
+	GENERATED_USTRUCT_BODY()
 
 	FWidgetCarouselNavigationButtonStyle();
 
@@ -37,7 +37,7 @@ struct WIDGETCAROUSEL_API FWidgetCarouselNavigationButtonStyle : public FSlateWi
 USTRUCT()
 struct WIDGETCAROUSEL_API FWidgetCarouselNavigationBarStyle : public FSlateWidgetStyle
 {
-	GENERATED_BODY()
+	GENERATED_USTRUCT_BODY()
 
 	FWidgetCarouselNavigationBarStyle();
 
@@ -91,4 +91,4 @@ private:
 private:
 
 	static TSharedPtr< class FSlateStyleSet > WidgetCarouselStyleInstance;
-};
+};
\ No newline at end of file
diff --git a/Engine/Source/Developer/Windows/WindowsTargetPlatform/Classes/WindowsTargetSettings.h b/Engine/Source/Developer/Windows/WindowsTargetPlatform/Classes/WindowsTargetSettings.h
--- a/Engine/Source/Developer/Windows/WindowsTargetPlatform/Classes/WindowsTargetSettings.h
+++ b/Engine/Source/Developer/Windows/WindowsTargetPlatform/Classes/WindowsTargetSettings.h
@@ -22,11 +22,9 @@ UCLASS(config=Engine, defaultconfig)
 class WINDOWSTARGETPLATFORM_API UWindowsTargetSettings
 	: public UObject
 {
-	GENERATED_BODY()
-public:
-	UWindowsTargetSettings(const FObjectInitializer& ObjectInitializer = FObjectInitializer::Get());
 public:
 
+	GENERATED_UCLASS_BODY()
 
 	/** 
 	 * The collection of RHI's we want to support on this platform.
diff --git a/Engine/Source/Editor/AIGraph/Classes/AIGraph.h b/Engine/Source/Editor/AIGraph/Classes/AIGraph.h
--- a/Engine/Source/Editor/AIGraph/Classes/AIGraph.h
+++ b/Engine/Source/Editor/AIGraph/Classes/AIGraph.h
@@ -6,9 +6,7 @@
 UCLASS()
 class AIGRAPH_API UAIGraph : public UEdGraph
 {
-	GENERATED_BODY()
-public:
-	UAIGraph(const FObjectInitializer& ObjectInitializer = FObjectInitializer::Get());
+	GENERATED_UCLASS_BODY()
 
 	UPROPERTY()
 	int32 GraphVersion;
diff --git a/Engine/Source/Editor/AIGraph/Classes/AIGraphNode.h b/Engine/Source/Editor/AIGraph/Classes/AIGraphNode.h
--- a/Engine/Source/Editor/AIGraph/Classes/AIGraphNode.h
+++ b/Engine/Source/Editor/AIGraph/Classes/AIGraphNode.h
@@ -9,9 +9,7 @@
 UCLASS()
 class AIGRAPH_API UAIGraphNode : public UEdGraphNode
 {
-	GENERATED_BODY()
-public:
-	UAIGraphNode(const FObjectInitializer& ObjectInitializer = FObjectInitializer::Get());
+	GENERATED_UCLASS_BODY()
 
 	/** instance class */
 	UPROPERTY()
diff --git a/Engine/Source/Editor/AnimGraph/Classes/AnimGraphNode_ApplyAdditive.h b/Engine/Source/Editor/AnimGraph/Classes/AnimGraphNode_ApplyAdditive.h
--- a/Engine/Source/Editor/AnimGraph/Classes/AnimGraphNode_ApplyAdditive.h
+++ b/Engine/Source/Editor/AnimGraph/Classes/AnimGraphNode_ApplyAdditive.h
@@ -8,9 +8,7 @@
 UCLASS(MinimalAPI)
 class UAnimGraphNode_ApplyAdditive : public UAnimGraphNode_Base
 {
-	GENERATED_BODY()
-public:
-	ANIMGRAPH_API UAnimGraphNode_ApplyAdditive(const FObjectInitializer& ObjectInitializer = FObjectInitializer::Get());
+	GENERATED_UCLASS_BODY()
 
 	UPROPERTY(EditAnywhere, Category=Settings)
 	FAnimNode_ApplyAdditive Node;
diff --git a/Engine/Source/Editor/AnimGraph/Classes/AnimGraphNode_Base.h b/Engine/Source/Editor/AnimGraph/Classes/AnimGraphNode_Base.h
--- a/Engine/Source/Editor/AnimGraph/Classes/AnimGraphNode_Base.h
+++ b/Engine/Source/Editor/AnimGraph/Classes/AnimGraphNode_Base.h
@@ -88,9 +88,7 @@ protected:
 UCLASS(Abstract, MinimalAPI)
 class UAnimGraphNode_Base : public UK2Node
 {
-	GENERATED_BODY()
-public:
-	ANIMGRAPH_API UAnimGraphNode_Base(const FObjectInitializer& ObjectInitializer = FObjectInitializer::Get());
+	GENERATED_UCLASS_BODY()
 
 	UPROPERTY(EditAnywhere, BlueprintReadWrite, Category=PinOptions, EditFixedSize)
 	TArray<FOptionalPinFromProperty> ShowPinForProperties;
diff --git a/Engine/Source/Editor/AnimGraph/Classes/AnimGraphNode_BlendListByBool.h b/Engine/Source/Editor/AnimGraph/Classes/AnimGraphNode_BlendListByBool.h
--- a/Engine/Source/Editor/AnimGraph/Classes/AnimGraphNode_BlendListByBool.h
+++ b/Engine/Source/Editor/AnimGraph/Classes/AnimGraphNode_BlendListByBool.h
@@ -8,9 +8,7 @@
 UCLASS(MinimalAPI)
 class UAnimGraphNode_BlendListByBool : public UAnimGraphNode_BlendListBase
 {
-	GENERATED_BODY()
-public:
-	ANIMGRAPH_API UAnimGraphNode_BlendListByBool(const FObjectInitializer& ObjectInitializer = FObjectInitializer::Get());
+	GENERATED_UCLASS_BODY()
 
 	UPROPERTY(EditAnywhere, Category=Settings)
 	FAnimNode_BlendListByBool Node;
diff --git a/Engine/Source/Editor/AnimGraph/Classes/AnimGraphNode_BlendListByEnum.h b/Engine/Source/Editor/AnimGraph/Classes/AnimGraphNode_BlendListByEnum.h
--- a/Engine/Source/Editor/AnimGraph/Classes/AnimGraphNode_BlendListByEnum.h
+++ b/Engine/Source/Editor/AnimGraph/Classes/AnimGraphNode_BlendListByEnum.h
@@ -9,9 +9,7 @@
 UCLASS(MinimalAPI)
 class UAnimGraphNode_BlendListByEnum : public UAnimGraphNode_BlendListBase, public INodeDependingOnEnumInterface
 {
-	GENERATED_BODY()
-public:
-	ANIMGRAPH_API UAnimGraphNode_BlendListByEnum(const FObjectInitializer& ObjectInitializer = FObjectInitializer::Get());
+	GENERATED_UCLASS_BODY()
 
 	UPROPERTY(EditAnywhere, Category=Settings)
 	FAnimNode_BlendListByEnum Node;
diff --git a/Engine/Source/Editor/AnimGraph/Classes/AnimGraphNode_BlendListByInt.h b/Engine/Source/Editor/AnimGraph/Classes/AnimGraphNode_BlendListByInt.h
--- a/Engine/Source/Editor/AnimGraph/Classes/AnimGraphNode_BlendListByInt.h
+++ b/Engine/Source/Editor/AnimGraph/Classes/AnimGraphNode_BlendListByInt.h
@@ -8,9 +8,7 @@
 UCLASS(MinimalAPI)
 class UAnimGraphNode_BlendListByInt : public UAnimGraphNode_BlendListBase
 {
-	GENERATED_BODY()
-public:
-	ANIMGRAPH_API UAnimGraphNode_BlendListByInt(const FObjectInitializer& ObjectInitializer = FObjectInitializer::Get());
+	GENERATED_UCLASS_BODY()
 
 	UPROPERTY(EditAnywhere, Category=Settings)
 	FAnimNode_BlendListByInt Node;
