Messaging: Removed the WithMessaging type trait and message type pre-cache

#UpgradeNotes:
The TStructOpsTypeTraits<>::WithMessaging type trait is no longer necessary for UStructs intended to be used as messages, because UObject creation and destruction is now thread-safe, and message types no longer need to be pre-cached. You can now use any existing and custom UStruct class in the Engine as messages. I ran the Editor for a couple hours with the profiler enabled and did not have any stability issues. If you experience any Messaging related crashes, please let me know asap, thanks!

diff --git a/Engine/Plugins/Messaging/UdpMessaging/Source/UdpMessaging/Private/Transport/UdpDeserializedMessage.cpp b/Engine/Plugins/Messaging/UdpMessaging/Source/UdpMessaging/Private/Transport/UdpDeserializedMessage.cpp
--- a/Engine/Plugins/Messaging/UdpMessaging/Source/UdpMessaging/Private/Transport/UdpDeserializedMessage.cpp
+++ b/Engine/Plugins/Messaging/UdpMessaging/Source/UdpMessaging/Private/Transport/UdpDeserializedMessage.cpp
@@ -21,7 +21,8 @@ bool FUdpDeserializedMessage::Deserialize( const FUdpReassembledMessageRef& Reas
 		FName MessageType;
 		MessageReader << MessageType;
 
-		TypeInfo = FMessageTypeMap::MessageTypeMap.Find(MessageType.ToString());
+		// @todo gmp: cache message types for faster lookup
+		TypeInfo = FindObjectSafe<UScriptStruct>(ANY_PACKAGE, *MessageType.ToString());
 
 		if (!TypeInfo.IsValid(false, true))
 		{
diff --git a/Engine/Source/Developer/ProfilerMessages/Classes/ProfilerServiceMessages.h b/Engine/Source/Developer/ProfilerMessages/Classes/ProfilerServiceMessages.h
--- a/Engine/Source/Developer/ProfilerMessages/Classes/ProfilerServiceMessages.h
+++ b/Engine/Source/Developer/ProfilerMessages/Classes/ProfilerServiceMessages.h
@@ -38,15 +38,6 @@ struct FProfilerServiceAuthorize
 
 };
 
-template<>
-struct TStructOpsTypeTraits<FProfilerServiceAuthorize> : public TStructOpsTypeTraitsBase
-{
-	enum
-	{
-		WithMessageHandling = true
-	};
-};
-
 
 /**
  */
@@ -87,15 +78,6 @@ struct FProfilerServiceAuthorize2
 
 };
 
-template<>
-struct TStructOpsTypeTraits<FProfilerServiceAuthorize2> : public TStructOpsTypeTraitsBase
-{
-	enum
-	{
-		WithMessageHandling = true
-	};
-};
-
 
 /**
  */
@@ -130,15 +112,6 @@ struct FProfilerServiceData
 	}
 };
 
-template<>
-struct TStructOpsTypeTraits<FProfilerServiceData> : public TStructOpsTypeTraitsBase
-{
-	enum
-	{
-		WithMessageHandling = true
-	};
-};
-
 
 /**
  */
@@ -177,15 +150,6 @@ struct FProfilerServiceData2
 	}
 };
 
-template<>
-struct TStructOpsTypeTraits<FProfilerServiceData2> : public TStructOpsTypeTraitsBase
-{
-	enum
-	{
-		WithMessageHandling = true
-	};
-};
-
 
 /**
  */
@@ -218,15 +182,6 @@ struct FProfilerServicePreviewAck
 	}
 };
 
-template<>
-struct TStructOpsTypeTraits<FProfilerServicePreviewAck> : public TStructOpsTypeTraitsBase
-{
-	enum
-	{
-		WithMessageHandling = true
-	};
-};
-
 
 /**
  */
@@ -260,15 +215,6 @@ struct FProfilerServiceMetaData
 	{}
 };
 
-template<>
-struct TStructOpsTypeTraits<FProfilerServiceMetaData> : public TStructOpsTypeTraitsBase
-{
-	enum
-	{
-		WithMessageHandling = true
-	};
-};
-
 
 /**
  * Implements a message for copying a file through the network, as well as for synchronization.
@@ -331,14 +277,6 @@ struct FProfilerServiceFileChunk
 	{}
 };
 
-template<>
-struct TStructOpsTypeTraits<FProfilerServiceFileChunk> : public TStructOpsTypeTraitsBase
-{
-	enum
-	{
-		WithMessageHandling = true
-	};
-};
 
 /**
  */
@@ -348,15 +286,6 @@ struct FProfilerServicePing
 	GENERATED_USTRUCT_BODY()
 };
 
-template<>
-struct TStructOpsTypeTraits<FProfilerServicePing> : public TStructOpsTypeTraitsBase
-{
-	enum
-	{
-		WithMessageHandling = true
-	};
-};
-
 
 /**
  */
@@ -367,15 +296,6 @@ struct FProfilerServicePong
 
 };
 
-template<>
-struct TStructOpsTypeTraits<FProfilerServicePong> : public TStructOpsTypeTraitsBase
-{
-	enum
-	{
-		WithMessageHandling = true
-	};
-};
-
 
 /**
  */
@@ -408,15 +328,6 @@ struct FProfilerServiceSubscribe
 	{ }
 };
 
-template<>
-struct TStructOpsTypeTraits<FProfilerServiceSubscribe> : public TStructOpsTypeTraitsBase
-{
-	enum
-	{
-		WithMessageHandling = true
-	};
-};
-
 
 /**
  */
@@ -450,15 +361,6 @@ struct FProfilerServiceUnsubscribe
 	{ }
 };
 
-template<>
-struct TStructOpsTypeTraits<FProfilerServiceUnsubscribe> : public TStructOpsTypeTraitsBase
-{
-	enum
-	{
-		WithMessageHandling = true
-	};
-};
-
 
 /**
  */
@@ -486,15 +388,6 @@ struct FProfilerServiceCapture
 	{ }
 };
 
-template<>
-struct TStructOpsTypeTraits<FProfilerServiceCapture> : public TStructOpsTypeTraitsBase
-{
-	enum
-	{
-		WithMessageHandling = true
-	};
-};
-
 
 /**
  */
@@ -522,15 +415,6 @@ struct FProfilerServicePreview
 	{ }
 };
 
-template<>
-struct TStructOpsTypeTraits<FProfilerServicePreview> : public TStructOpsTypeTraitsBase
-{
-	enum
-	{
-		WithMessageHandling = true
-	};
-};
-
 
 /**
  */
@@ -555,12 +439,3 @@ struct FProfilerServiceRequest
 		: Request(InRequest)
 	{ }
 };
-
-template<>
-struct TStructOpsTypeTraits<FProfilerServiceRequest> : public TStructOpsTypeTraitsBase
-{
-	enum
-	{
-		WithMessageHandling = true
-	};
-};
diff --git a/Engine/Source/Developer/SessionMessages/Classes/SessionServiceMessages.h b/Engine/Source/Developer/SessionMessages/Classes/SessionServiceMessages.h
--- a/Engine/Source/Developer/SessionMessages/Classes/SessionServiceMessages.h
+++ b/Engine/Source/Developer/SessionMessages/Classes/SessionServiceMessages.h
@@ -18,15 +18,6 @@ struct FSessionServicePing
 
 };
 
-template<>
-struct TStructOpsTypeTraits<FSessionServicePing> : public TStructOpsTypeTraitsBase
-{
-	enum
-	{
-		WithMessageHandling = true
-	};
-};
-
 
 /**
  * Implements a message that is published in response to FSessionServicePing.
@@ -77,15 +68,6 @@ struct FSessionServicePong
 	bool Standalone;
 };
 
-template<>
-struct TStructOpsTypeTraits<FSessionServicePong> : public TStructOpsTypeTraitsBase
-{
-	enum
-	{
-		WithMessageHandling = true
-	};
-};
-
 
 /* Session status messages
  *****************************************************************************/
@@ -137,15 +119,6 @@ public:
 	{ }
 };
 
-template<>
-struct TStructOpsTypeTraits<FSessionServiceLog> : public TStructOpsTypeTraitsBase
-{
-	enum
-	{
-		WithMessageHandling = true
-	};
-};
-
 
 /**
  * Implements a message to subscribe to an application's console log.
@@ -156,15 +129,6 @@ struct FSessionServiceLogSubscribe
 	GENERATED_USTRUCT_BODY()
 };
 
-template<>
-struct TStructOpsTypeTraits<FSessionServiceLogSubscribe> : public TStructOpsTypeTraitsBase
-{
-	enum
-	{
-		WithMessageHandling = true
-	};
-};
-
 
 /**
  * Implements a message to unsubscribe from an application's console log.
@@ -174,12 +138,3 @@ struct FSessionServiceLogUnsubscribe
 {
 	GENERATED_USTRUCT_BODY()
 };
-
-template<>
-struct TStructOpsTypeTraits<FSessionServiceLogUnsubscribe> : public TStructOpsTypeTraitsBase
-{
-	enum
-	{
-		WithMessageHandling = true
-	};
-};
diff --git a/Engine/Source/Developer/TargetDeviceServices/Classes/TargetDeviceServiceMessages.h b/Engine/Source/Developer/TargetDeviceServices/Classes/TargetDeviceServiceMessages.h
--- a/Engine/Source/Developer/TargetDeviceServices/Classes/TargetDeviceServiceMessages.h
+++ b/Engine/Source/Developer/TargetDeviceServices/Classes/TargetDeviceServiceMessages.h
@@ -37,15 +37,6 @@ struct FTargetDeviceServiceDeployCommit
 	{ }
 };
 
-template<>
-struct TStructOpsTypeTraits<FTargetDeviceServiceDeployCommit> : public TStructOpsTypeTraitsBase
-{
-	enum
-	{
-		WithMessageHandling = true
-	};
-};
-
 
 /**
  * Implements a message for deploying a single file to a target device.
@@ -78,15 +69,6 @@ struct FTargetDeviceServiceDeployFile
 	{ }
 };
 
-template<>
-struct TStructOpsTypeTraits<FTargetDeviceServiceDeployFile> : public TStructOpsTypeTraitsBase
-{
-	enum
-	{
-		WithMessageHandling = true
-	};
-};
-
 
 /**
  * Implements a message for notifying a target device proxy that a deployment transaction has finished.
@@ -134,15 +116,6 @@ struct FTargetDeviceServiceDeployFinished
 	{ }
 };
 
-template<>
-struct TStructOpsTypeTraits<FTargetDeviceServiceDeployFinished> : public TStructOpsTypeTraitsBase
-{
-	enum
-	{
-		WithMessageHandling = true
-	};
-};
-
 
 /* Application launch messages
 *****************************************************************************/
@@ -195,15 +168,6 @@ struct FTargetDeviceServiceLaunchApp
 	{ }
 };
 
-template<>
-struct TStructOpsTypeTraits<FTargetDeviceServiceLaunchApp> : public TStructOpsTypeTraitsBase
-{
-	enum
-	{
-		WithMessageHandling = true
-	};
-};
-
 
 /**
  * Implements a message for notifying a target device proxy that launching an application has finished.
@@ -246,15 +210,6 @@ struct FTargetDeviceServiceLaunchFinished
 	{ }
 };
 
-template<>
-struct TStructOpsTypeTraits<FTargetDeviceServiceLaunchFinished> : public TStructOpsTypeTraitsBase
-{
-	enum
-	{
-		WithMessageHandling = true
-	};
-};
-
 
 /* Device claiming messages
  *****************************************************************************/
@@ -294,15 +249,6 @@ struct FTargetDeviceClaimDenied
 	{ }
 };
 
-template<>
-struct TStructOpsTypeTraits<FTargetDeviceClaimDenied> : public TStructOpsTypeTraitsBase
-{
-	enum
-	{
-		WithMessageHandling = true
-	};
-};
-
 
 /**
  * Implements a message that is sent when a service claimed a device.
@@ -339,15 +285,6 @@ struct FTargetDeviceClaimed
 	{ }
 };
 
-template<>
-struct TStructOpsTypeTraits<FTargetDeviceClaimed> : public TStructOpsTypeTraitsBase
-{
-	enum
-	{
-		WithMessageHandling = true
-	};
-};
-
 
 /**
  * Implements a message that is sent when a device is no longer claimed.
@@ -383,15 +320,6 @@ struct FTargetDeviceUnclaimed
 	{ }
 };
 
-template<>
-struct TStructOpsTypeTraits<FTargetDeviceUnclaimed> : public TStructOpsTypeTraitsBase
-{
-	enum
-	{
-		WithMessageHandling = true
-	};
-};
-
 
 /* Device discovery messages
  *****************************************************************************/
@@ -418,15 +346,6 @@ struct FTargetDeviceServicePing
 	{ }
 };
 
-template<>
-struct TStructOpsTypeTraits<FTargetDeviceServicePing> : public TStructOpsTypeTraitsBase
-{
-	enum
-	{
-		WithMessageHandling = true
-	};
-};
-
 
 /**
 * Struct for a flavor's information
@@ -455,15 +374,6 @@ struct FTargetDeviceVariant
 	FString PlatformDisplayName;
 };
 
-template<>
-struct TStructOpsTypeTraits<FTargetDeviceVariant> : public TStructOpsTypeTraitsBase
-{
-	enum
-	{
-		WithMessageHandling = true
-	};
-};
-
 
 /**
  * Implements a message that is sent in response to target device service discovery messages.
@@ -542,15 +452,6 @@ struct FTargetDeviceServicePong
 	TArray<FTargetDeviceVariant> Variants;
 };
 
-template<>
-struct TStructOpsTypeTraits<FTargetDeviceServicePong> : public TStructOpsTypeTraitsBase
-{
-	enum
-	{
-		WithMessageHandling = true
-	};
-};
-
 
 /* Miscellaneous messages
  *****************************************************************************/
@@ -586,15 +487,6 @@ struct FTargetDeviceServicePowerOff
 	{ }
 };
 
-template<>
-struct TStructOpsTypeTraits<FTargetDeviceServicePowerOff> : public TStructOpsTypeTraitsBase
-{
-	enum
-	{
-		WithMessageHandling = true
-	};
-};
-
 
 /**
  * Implements a message for powering on a target device.
@@ -618,15 +510,6 @@ struct FTargetDeviceServicePowerOn
 	{ }
 };
 
-template<>
-struct TStructOpsTypeTraits<FTargetDeviceServicePowerOn> : public TStructOpsTypeTraitsBase
-{
-	enum
-	{
-		WithMessageHandling = true
-	};
-};
-
 
 /**
  * Implements a message for rebooting a target device.
@@ -650,15 +533,6 @@ struct FTargetDeviceServiceReboot
 	{ }
 };
 
-template<>
-struct TStructOpsTypeTraits<FTargetDeviceServiceReboot> : public TStructOpsTypeTraitsBase
-{
-	enum
-	{
-		WithMessageHandling = true
-	};
-};
-
 
 /**
  * Implements a message for running an executable on a target device.
@@ -696,15 +570,6 @@ struct FTargetDeviceServiceRunExecutable
 	{ }
 };
 
-template<>
-struct TStructOpsTypeTraits<FTargetDeviceServiceRunExecutable> : public TStructOpsTypeTraitsBase
-{
-	enum
-	{
-		WithMessageHandling = true
-	};
-};
-
 
 /**
  * Implements a message for notifying a target device proxy that running an executable has finished.
@@ -744,12 +609,3 @@ struct FTargetDeviceServiceRunFinished
 		, Succeeded(InSucceeded)
 	{ }
 };
-
-template<>
-struct TStructOpsTypeTraits<FTargetDeviceServiceRunFinished> : public TStructOpsTypeTraitsBase
-{
-	enum
-	{
-		WithMessageHandling = true
-	};
-};
diff --git a/Engine/Source/Editor/SwarmInterface/Private/SwarmMessages.h b/Engine/Source/Editor/SwarmInterface/Private/SwarmMessages.h
--- a/Engine/Source/Editor/SwarmInterface/Private/SwarmMessages.h
+++ b/Engine/Source/Editor/SwarmInterface/Private/SwarmMessages.h
@@ -12,15 +12,6 @@ struct FSwarmPingMessage
 	GENERATED_USTRUCT_BODY()
 };
 
-template<>
-struct TStructOpsTypeTraits<FSwarmPingMessage> : public TStructOpsTypeTraitsBase
-{
-	enum
-	{
-		WithMessageHandling = true
-	};
-};
-
 
 USTRUCT()
 struct FSwarmPongMessage
@@ -42,15 +33,6 @@ struct FSwarmPongMessage
 	}
 };
 
-template<>
-struct TStructOpsTypeTraits<FSwarmPongMessage> : public TStructOpsTypeTraitsBase
-{
-	enum
-	{
-		WithMessageHandling = true
-	};
-};
-
 
 USTRUCT()
 struct FSwarmInfoMessage
@@ -68,15 +50,6 @@ struct FSwarmInfoMessage
 	}
 };
 
-template<>
-struct TStructOpsTypeTraits<FSwarmInfoMessage> : public TStructOpsTypeTraitsBase
-{
-	enum
-	{
-		WithMessageHandling = true
-	};
-};
-
 
 USTRUCT()
 struct FSwarmAlertMessage
@@ -111,15 +84,6 @@ struct FSwarmAlertMessage
 	}
 };
 
-template<>
-struct TStructOpsTypeTraits<FSwarmAlertMessage> : public TStructOpsTypeTraitsBase
-{
-	enum
-	{
-		WithMessageHandling = true
-	};
-};
-
 
 USTRUCT()
 struct FSwarmTimingMessage
@@ -142,15 +106,6 @@ struct FSwarmTimingMessage
 	}
 };
 
-template<>
-struct TStructOpsTypeTraits<FSwarmTimingMessage> : public TStructOpsTypeTraitsBase
-{
-	enum
-	{
-		WithMessageHandling = true
-	};
-};
-
 
 USTRUCT()
 struct FSwarmTaskRequestMessage
@@ -160,15 +115,6 @@ struct FSwarmTaskRequestMessage
 	FSwarmTaskRequestMessage() {}
 };
 
-template<>
-struct TStructOpsTypeTraits<FSwarmTaskRequestMessage> : public TStructOpsTypeTraitsBase
-{
-	enum
-	{
-		WithMessageHandling = true
-	};
-};
-
 
 USTRUCT()
 struct FSwarmTaskRequestReleaseMessage
@@ -178,15 +124,6 @@ struct FSwarmTaskRequestReleaseMessage
 	FSwarmTaskRequestReleaseMessage() {}
 };
 
-template<>
-struct TStructOpsTypeTraits<FSwarmTaskRequestReleaseMessage> : public TStructOpsTypeTraitsBase
-{
-	enum
-	{
-		WithMessageHandling = true
-	};
-};
-
 
 USTRUCT()
 struct FSwarmTaskRequestReservationMessage
@@ -196,15 +133,6 @@ struct FSwarmTaskRequestReservationMessage
 	FSwarmTaskRequestReservationMessage() {}
 };
 
-template<>
-struct TStructOpsTypeTraits<FSwarmTaskRequestReservationMessage> : public TStructOpsTypeTraitsBase
-{
-	enum
-	{
-		WithMessageHandling = true
-	};
-};
-
 
 USTRUCT()
 struct FSwarmTaskRequestSpecificationMessage
@@ -237,15 +165,6 @@ struct FSwarmTaskRequestSpecificationMessage
 	{ }
 };
 
-template<>
-struct TStructOpsTypeTraits<FSwarmTaskRequestSpecificationMessage> : public TStructOpsTypeTraitsBase
-{
-	enum
-	{
-		WithMessageHandling = true
-	};
-};
-
 
 USTRUCT()
 struct FSwarmJobStateMessage
@@ -278,15 +197,6 @@ struct FSwarmJobStateMessage
 	}
 };
 
-template<>
-struct TStructOpsTypeTraits<FSwarmJobStateMessage> : public TStructOpsTypeTraitsBase
-{
-	enum
-	{
-		WithMessageHandling = true
-	};
-};
-
 
 USTRUCT()
 struct FSwarmTaskStateMessage
@@ -319,15 +229,6 @@ struct FSwarmTaskStateMessage
 	}
 };
 
-template<>
-struct TStructOpsTypeTraits<FSwarmTaskStateMessage> : public TStructOpsTypeTraitsBase
-{
-	enum
-	{
-		WithMessageHandling = true
-	};
-};
-
 
 USTRUCT()
 struct FSwarmQuitMessage
@@ -336,12 +237,3 @@ struct FSwarmQuitMessage
 
 	FSwarmQuitMessage() {}
 };
-
-template<>
-struct TStructOpsTypeTraits<FSwarmQuitMessage> : public TStructOpsTypeTraitsBase
-{
-	enum
-	{
-		WithMessageHandling = true
-	};
-};
diff --git a/Engine/Source/Editor/UnrealEdMessages/Classes/AssetEditorMessages.h b/Engine/Source/Editor/UnrealEdMessages/Classes/AssetEditorMessages.h
--- a/Engine/Source/Editor/UnrealEdMessages/Classes/AssetEditorMessages.h
+++ b/Engine/Source/Editor/UnrealEdMessages/Classes/AssetEditorMessages.h
@@ -32,13 +32,3 @@ struct FAssetEditorRequestOpenAsset
 		: AssetName(InAssetName)
 	{ }
 };
-
-
-template<>
-struct TStructOpsTypeTraits<FAssetEditorRequestOpenAsset> : public TStructOpsTypeTraitsBase
-{
-	enum
-	{
-		WithMessageHandling = true
-	};
-};
diff --git a/Engine/Source/Editor/UnrealEdMessages/Classes/FileServerMessages.h b/Engine/Source/Editor/UnrealEdMessages/Classes/FileServerMessages.h
--- a/Engine/Source/Editor/UnrealEdMessages/Classes/FileServerMessages.h
+++ b/Engine/Source/Editor/UnrealEdMessages/Classes/FileServerMessages.h
@@ -39,13 +39,3 @@ struct FFileServerReady
 		, InstanceId(InInstanceId)
 	{ }
 };
-
-
-template<>
-struct TStructOpsTypeTraits<FFileServerReady> : public TStructOpsTypeTraitsBase
-{
-	enum
-	{
-		WithMessageHandling = true
-	};
-};
diff --git a/Engine/Source/Runtime/AutomationMessages/Classes/AutomationWorkerMessages.h b/Engine/Source/Runtime/AutomationMessages/Classes/AutomationWorkerMessages.h
--- a/Engine/Source/Runtime/AutomationMessages/Classes/AutomationWorkerMessages.h
+++ b/Engine/Source/Runtime/AutomationMessages/Classes/AutomationWorkerMessages.h
@@ -45,16 +45,6 @@ struct FAutomationWorkerFindWorkers
 };
 
 
-template<>
-struct TStructOpsTypeTraits<FAutomationWorkerFindWorkers> : public TStructOpsTypeTraitsBase
-{
-	enum
-	{
-		WithMessageHandling = true
-	};
-};
-
-
 /**
  * Implements a message that is sent in response to FAutomationWorkerFindWorkers.
  */
@@ -109,15 +99,6 @@ struct FAutomationWorkerFindWorkersResponse
 	FAutomationWorkerFindWorkersResponse( ) { }
 };
 
-template<>
-struct TStructOpsTypeTraits<FAutomationWorkerFindWorkersResponse> : public TStructOpsTypeTraitsBase
-{
-	enum
-	{
-		WithMessageHandling = true
-	};
-};
-
 
 /**
  * Implements a message that notifies automation controllers that a worker went off-line.
@@ -128,15 +109,6 @@ struct FAutomationWorkerWorkerOffline
 	GENERATED_USTRUCT_BODY()
 };
 
-template<>
-struct TStructOpsTypeTraits<FAutomationWorkerWorkerOffline> : public TStructOpsTypeTraitsBase
-{
-	enum
-	{
-		WithMessageHandling = true
-	};
-};
-
 
 /**
  */
@@ -146,15 +118,6 @@ struct FAutomationWorkerPing
 	GENERATED_USTRUCT_BODY()
 };
 
-template<>
-struct TStructOpsTypeTraits<FAutomationWorkerPing> : public TStructOpsTypeTraitsBase
-{
-	enum
-	{
-		WithMessageHandling = true
-	};
-};
-
 
 /**
  */
@@ -164,15 +127,6 @@ struct FAutomationWorkerResetTests
 	GENERATED_USTRUCT_BODY()
 };
 
-template<>
-struct TStructOpsTypeTraits<FAutomationWorkerResetTests> : public TStructOpsTypeTraitsBase
-{
-	enum
-	{
-		WithMessageHandling = true
-	};
-};
-
 
 /**
  */
@@ -182,15 +136,6 @@ struct FAutomationWorkerPong
 	GENERATED_USTRUCT_BODY()
 };
 
-template<>
-struct TStructOpsTypeTraits<FAutomationWorkerPong> : public TStructOpsTypeTraitsBase
-{
-	enum
-	{
-		WithMessageHandling = true
-	};
-};
-
 
 /**
  * Implements a message for requesting available automation tests from a worker.
@@ -222,15 +167,6 @@ struct FAutomationWorkerRequestTests
 	{ }
 };
 
-template<>
-struct TStructOpsTypeTraits<FAutomationWorkerRequestTests> : public TStructOpsTypeTraitsBase
-{
-	enum
-	{
-		WithMessageHandling = true
-	};
-};
-
 
 /**
  * Implements a message that is sent in response to FAutomationWorkerRequestTests.
@@ -262,15 +198,6 @@ struct FAutomationWorkerRequestTestsReply
 	{ }
 };
 
-template<>
-struct TStructOpsTypeTraits<FAutomationWorkerRequestTestsReply> : public TStructOpsTypeTraitsBase
-{
-	enum
-	{
-		WithMessageHandling = true
-	};
-};
-
 
 /**
  * Implements a message to request the running of automation tests on a worker.
@@ -317,15 +244,6 @@ struct FAutomationWorkerRunTests
 	{ }
 };
 
-template<>
-struct TStructOpsTypeTraits<FAutomationWorkerRunTests> : public TStructOpsTypeTraitsBase
-{
-	enum
-	{
-		WithMessageHandling = true
-	};
-};
-
 
 /**
  * Implements a message that is sent in response to FAutomationWorkerRunTests.
@@ -364,15 +282,6 @@ struct FAutomationWorkerRunTestsReply
 	TArray<FString> Warnings;
 };
 
-template<>
-struct TStructOpsTypeTraits<FAutomationWorkerRunTestsReply> : public TStructOpsTypeTraitsBase
-{
-	enum
-	{
-		WithMessageHandling = true
-	};
-};
-
 
 /**
  */
@@ -398,15 +307,6 @@ struct FAutomationWorkerRequestNextNetworkCommand
 	{ }
 };
 
-template<>
-struct TStructOpsTypeTraits<FAutomationWorkerRequestNextNetworkCommand> : public TStructOpsTypeTraitsBase
-{
-	enum
-	{
-		WithMessageHandling = true
-	};
-};
-
 
 /**
  */
@@ -416,15 +316,6 @@ struct FAutomationWorkerNextNetworkCommandReply
 	GENERATED_USTRUCT_BODY()
 };
 
-template<>
-struct TStructOpsTypeTraits<FAutomationWorkerNextNetworkCommandReply> : public TStructOpsTypeTraitsBase
-{
-	enum
-	{
-		WithMessageHandling = true
-	};
-};
-
 
 /**
  * Implements a message that is sent in containing a screen shot run during performance test.
@@ -442,12 +333,3 @@ struct FAutomationWorkerScreenImage
 	UPROPERTY()
 	FString ScreenShotName;
 };
-
-template<>
-struct TStructOpsTypeTraits<FAutomationWorkerScreenImage> : public TStructOpsTypeTraitsBase
-{
-	enum
-	{
-		WithMessageHandling = true
-	};
-};
diff --git a/Engine/Source/Runtime/CoreUObject/Private/CoreUObjectPrivate.h b/Engine/Source/Runtime/CoreUObject/Private/CoreUObjectPrivate.h
--- a/Engine/Source/Runtime/CoreUObject/Private/CoreUObjectPrivate.h
+++ b/Engine/Source/Runtime/CoreUObject/Private/CoreUObjectPrivate.h
@@ -32,7 +32,6 @@
 #include "Linker.h"						// Linker.
 #include "GCObject.h"			        // non-UObject object referencer
 #include "RedirectCollector.h"
-#include "MessageTypeMap.h"
 #include "WorldCompositionUtility.h"
 #include "StringClassReference.h"
 #include "Blueprint/BlueprintSupport.h"
diff --git a/Engine/Source/Runtime/CoreUObject/Private/Misc/MessageTypeMap.cpp b/Engine/Source/Runtime/CoreUObject/Private/Misc/MessageTypeMap.cpp
deleted file mode 100644
diff --git a/Engine/Source/Runtime/CoreUObject/Private/UObject/Class.cpp b/Engine/Source/Runtime/CoreUObject/Private/UObject/Class.cpp
--- a/Engine/Source/Runtime/CoreUObject/Private/UObject/Class.cpp
+++ b/Engine/Source/Runtime/CoreUObject/Private/UObject/Class.cpp
@@ -2044,13 +2044,6 @@ void UScriptStruct::PrepareCppStructOps()
 				UE_LOG(LogClass, Verbose, TEXT("Native struct %s has native SerializeFromMismatchedTag."),*GetName());
 				StructFlags = EStructFlags(StructFlags | STRUCT_SerializeFromMismatchedTag);
 			}
-			if (CppStructOps->HasMessageHandling())
-			{
-				UE_LOG(LogClass, Verbose, TEXT("Natice struct %s is a message type."), *GetName());
-
-				// add to the message type registry
-				FMessageTypeMap::MessageTypeMap.Add(GetName(), this);
-			}
 		}
 	}
 	check(!bPrepareCppStructOpsCompleted); // recursion is unacceptable
diff --git a/Engine/Source/Runtime/CoreUObject/Public/CoreUObject.h b/Engine/Source/Runtime/CoreUObject/Public/CoreUObject.h
--- a/Engine/Source/Runtime/CoreUObject/Public/CoreUObject.h
+++ b/Engine/Source/Runtime/CoreUObject/Public/CoreUObject.h
@@ -38,7 +38,6 @@
 #include "NotifyHook.h"
 #include "RedirectCollector.h"
 #include "ScriptStackTracker.h"
-#include "MessageTypeMap.h"
 #include "WorldCompositionUtility.h"
 #include "StringClassReference.h"
 
diff --git a/Engine/Source/Runtime/CoreUObject/Public/Misc/MessageTypeMap.h b/Engine/Source/Runtime/CoreUObject/Public/Misc/MessageTypeMap.h
deleted file mode 100644
diff --git a/Engine/Source/Runtime/CoreUObject/Public/UObject/Class.h b/Engine/Source/Runtime/CoreUObject/Public/UObject/Class.h
--- a/Engine/Source/Runtime/CoreUObject/Public/UObject/Class.h
+++ b/Engine/Source/Runtime/CoreUObject/Public/UObject/Class.h
@@ -490,7 +490,6 @@ struct TStructOpsTypeTraitsBase
 		WithNetSerializer              = false, // struct has a NetSerialize function for serializing its state to an FArchive used for network replication.
 		WithNetDeltaSerializer         = false, // struct has a NetDeltaSerialize function for serializing differences in state from a previous NetSerialize operation.
 		WithSerializeFromMismatchedTag = false, // struct has a SerializeFromMismatchedTag function for converting from other property tags.
-		WithMessageHandling            = false, // struct can be located by the message bus system for marshalling and unmarshalling.
 	};
 };
 
@@ -821,8 +820,6 @@ public:
 		**/
 		virtual bool SerializeFromMismatchedTag(struct FPropertyTag const& Tag, FArchive& Ar, void *Data) = 0;
 
-		virtual bool HasMessageHandling() = 0;
-
 	private:
 		/** sizeof() of the structure **/
 		const int32 Size;
@@ -958,10 +955,6 @@ public:
 			check(TTraits::WithSerializeFromMismatchedTag); // don't call this if we have indicated it is not allowed
 			return SerializeFromMismatchedTagOrNot(Tag, Ar, (CPPSTRUCT*)Data);
 		}
-		virtual bool HasMessageHandling() override
-		{
-			return TTraits::WithMessageHandling;
-		}
 	};
 
 	/** Template for noexport classes to autoregister before main starts **/
diff --git a/Engine/Source/Runtime/EngineMessages/Classes/EngineServiceMessages.h b/Engine/Source/Runtime/EngineMessages/Classes/EngineServiceMessages.h
--- a/Engine/Source/Runtime/EngineMessages/Classes/EngineServiceMessages.h
+++ b/Engine/Source/Runtime/EngineMessages/Classes/EngineServiceMessages.h
@@ -17,15 +17,6 @@ struct FEngineServicePing
 	GENERATED_USTRUCT_BODY()
 };
 
-template<>
-struct TStructOpsTypeTraits<FEngineServicePing> : public TStructOpsTypeTraitsBase
-{
-	enum
-	{
-		WithMessageHandling = true
-	};
-};
-
 
 /**
  * Implements a message for responding to a request to discover engine instances on the network.
@@ -64,15 +55,6 @@ struct FEngineServicePong
 	float WorldTimeSeconds;
 };
 
-template<>
-struct TStructOpsTypeTraits<FEngineServicePong> : public TStructOpsTypeTraitsBase
-{
-	enum
-	{
-		WithMessageHandling = true
-	};
-};
-
 
 /* Authorization messages
  *****************************************************************************/
@@ -94,15 +76,6 @@ struct FEngineServiceAuthDeny
 	FString UserToDeny;
 };
 
-template<>
-struct TStructOpsTypeTraits<FEngineServiceAuthDeny> : public TStructOpsTypeTraitsBase
-{
-	enum
-	{
-		WithMessageHandling = true
-	};
-};
-
 
 /**
  * Implements a message for granting service access to a remote user.
@@ -122,15 +95,6 @@ struct FEngineServiceAuthGrant
 
 };
 
-template<>
-struct TStructOpsTypeTraits<FEngineServiceAuthGrant> : public TStructOpsTypeTraitsBase
-{
-	enum
-	{
-		WithMessageHandling = true
-	};
-};
-
 
 /* Command messages
  *****************************************************************************/
@@ -165,15 +129,6 @@ struct FEngineServiceExecuteCommand
 	{ }
 };
 
-template<>
-struct TStructOpsTypeTraits<FEngineServiceExecuteCommand> : public TStructOpsTypeTraitsBase
-{
-	enum
-	{
-		WithMessageHandling = true
-	};
-};
-
 
 /**
  * Implements a message for terminating the engine.
@@ -200,15 +155,6 @@ struct FEngineServiceTerminate
 	{ }
 };
 
-template<>
-struct TStructOpsTypeTraits<FEngineServiceTerminate> : public TStructOpsTypeTraitsBase
-{
-	enum
-	{
-		WithMessageHandling = true
-	};
-};
-
 
 /* Status messages
  *****************************************************************************/
@@ -242,12 +188,3 @@ struct FEngineServiceNotification
 		, TimeSeconds(InTimeSeconds)
 	{ }
 };
-
-template<>
-struct TStructOpsTypeTraits<FEngineServiceNotification> : public TStructOpsTypeTraitsBase
-{
-	enum
-	{
-		WithMessageHandling = true
-	};
-};
diff --git a/Engine/Source/Runtime/IOS/LaunchDaemonMessages/Classes/IOSMessageProtocol.h b/Engine/Source/Runtime/IOS/LaunchDaemonMessages/Classes/IOSMessageProtocol.h
--- a/Engine/Source/Runtime/IOS/LaunchDaemonMessages/Classes/IOSMessageProtocol.h
+++ b/Engine/Source/Runtime/IOS/LaunchDaemonMessages/Classes/IOSMessageProtocol.h
@@ -13,15 +13,6 @@ struct FIOSLaunchDaemonPing
 	//FIOSLaunchDaemonPing() {}
 };
 
-template<>
-struct TStructOpsTypeTraits<FIOSLaunchDaemonPing> : public TStructOpsTypeTraitsBase
-{
-	enum
-	{
-		WithMessageHandling = true
-	};
-};
-
 
 USTRUCT()
 struct FIOSLaunchDaemonPong
@@ -62,15 +53,6 @@ struct FIOSLaunchDaemonPong
 	{}
 };
 
-template<>
-struct TStructOpsTypeTraits<FIOSLaunchDaemonPong> : public TStructOpsTypeTraitsBase
-{
-	enum
-	{
-		WithMessageHandling = true
-	};
-};
-
 
 USTRUCT()
 struct FIOSLaunchDaemonLaunchApp
@@ -90,12 +72,3 @@ struct FIOSLaunchDaemonLaunchApp
 		, Parameters(Params)
 	{}
 };
-
-template<>
-struct TStructOpsTypeTraits<FIOSLaunchDaemonLaunchApp> : public TStructOpsTypeTraitsBase
-{
-	enum
-	{
-		WithMessageHandling = true
-	};
-};
diff --git a/Engine/Source/Runtime/Messaging/Public/Helpers/MessageEndpointBuilder.h b/Engine/Source/Runtime/Messaging/Public/Helpers/MessageEndpointBuilder.h
--- a/Engine/Source/Runtime/Messaging/Public/Helpers/MessageEndpointBuilder.h
+++ b/Engine/Source/Runtime/Messaging/Public/Helpers/MessageEndpointBuilder.h
@@ -66,8 +66,6 @@ public:
 	template<typename MessageType, typename HandlerType>
 	FMessageEndpointBuilder& Handling( HandlerType* Handler, typename TRawMessageHandler<MessageType, HandlerType>::FuncType HandlerFunc )
 	{
-		// @todo gmp: implement proper async message deserialization, so this can be removed
-		static_assert(TStructOpsTypeTraits<MessageType>::WithMessageHandling == true, "Please add a WithMessageHandling type trait.");
 		Handlers.Add(MakeShareable(new TRawMessageHandler<MessageType, HandlerType>(Handler, HandlerFunc)));
 
 		return *this;
@@ -90,8 +88,6 @@ public:
 	template<typename MessageType>
 	FMessageEndpointBuilder& Handling( typename TFunctionMessageHandler<MessageType>::FuncType HandlerFunc )
 	{
-		// @todo gmp: implement proper async message deserialization, so this can be removed
-		static_assert(TStructOpsTypeTraits<MessageType>::WithMessageHandling == true, "Please add a WithMessageHandling type trait.");
 		Handlers.Add(MakeShareable(new TFunctionMessageHandler<MessageType>(HandlerFunc)));
 
 		return *this;
