Disable HandleClaimDeniedMessage().
Allows the editor to always claim a device which is desirable on as number of platforms.

diff --git a/Engine/Source/Developer/TargetDeviceServices/Private/TargetDeviceService.cpp b/Engine/Source/Developer/TargetDeviceServices/Private/TargetDeviceService.cpp
--- a/Engine/Source/Developer/TargetDeviceServices/Private/TargetDeviceService.cpp
+++ b/Engine/Source/Developer/TargetDeviceServices/Private/TargetDeviceService.cpp
@@ -255,14 +255,18 @@ bool FTargetDeviceService::StoreDeployedFile(FArchive* FileReader, const FString
 
 void FTargetDeviceService::HandleClaimDeniedMessage(const FTargetDeviceClaimDenied& Message, const IMessageContextRef& Context)
 {
-	if (Running && (Message.DeviceName == DeviceName))
-	{
-		Stop();
-
-		ClaimAddress = Context->GetSender();
-		ClaimHost = Message.HostName;
-		ClaimUser = Message.HostUser;
-	}
+// HACK: Disabling claim denied message. Allows the editor to always claim a device and should pervent cases where instances of the
+// editor running on other machines claim a device - which is undesirable on some platforms.
+#if 0
+ 	if (Running && (Message.DeviceName == DeviceName))
+ 	{
+ 		Stop();
+ 
+ 		ClaimAddress = Context->GetSender();
+ 		ClaimHost = Message.HostName;
+ 		ClaimUser = Message.HostUser;
+ 	}
+#endif
 }
 
 
