Additional temporary hack to hide remote devices.
Fix UE-18830. Comment typo fix.

diff --git a/Engine/Source/Developer/TargetDeviceServices/Private/TargetDeviceProxyManager.cpp b/Engine/Source/Developer/TargetDeviceServices/Private/TargetDeviceProxyManager.cpp
--- a/Engine/Source/Developer/TargetDeviceServices/Private/TargetDeviceProxyManager.cpp
+++ b/Engine/Source/Developer/TargetDeviceServices/Private/TargetDeviceProxyManager.cpp
@@ -118,6 +118,12 @@ void FTargetDeviceProxyManager::SendPing()
 
 void FTargetDeviceProxyManager::HandlePongMessage(const FTargetDeviceServicePong& Message, const IMessageContextRef& Context)
 {
+	// Another HACK: Ignore devices from other machines. See FTargetDeviceService::HandleClaimDeniedMessage()
+	if (Message.HostName != FPlatformProcess::ComputerName())
+ 	{
+ 		return;
+ 	}
+
 	TSharedPtr<FTargetDeviceProxy>& Proxy = Proxies.FindOrAdd(Message.Name);
 
 	if (!Proxy.IsValid())
diff --git a/Engine/Source/Developer/TargetDeviceServices/Private/TargetDeviceService.cpp b/Engine/Source/Developer/TargetDeviceServices/Private/TargetDeviceService.cpp
--- a/Engine/Source/Developer/TargetDeviceServices/Private/TargetDeviceService.cpp
+++ b/Engine/Source/Developer/TargetDeviceServices/Private/TargetDeviceService.cpp
@@ -255,8 +255,9 @@ bool FTargetDeviceService::StoreDeployedFile(FArchive* FileReader, const FString
 
 void FTargetDeviceService::HandleClaimDeniedMessage(const FTargetDeviceClaimDenied& Message, const IMessageContextRef& Context)
 {
-// HACK: Disabling claim denied message. Allows the editor to always claim a device and should pervent cases where instances of the
-// editor running on other machines claim a device - which is undesirable on some platforms.
+// HACK: Disabling claim denied message. Allows the editor to always claim a device and should prevent cases where instances of the
+// editor running on other machines can steal a device from us - which is undesirable on some platforms.
+// Also see FTargetDeviceProxyManager::HandlePongMessage()
 #if 0
  	if (Running && (Message.DeviceName == DeviceName))
  	{
