MessagingDebugger: Exposed message fields to details panel

diff --git a/Engine/Source/Developer/ProfilerMessages/Classes/ProfilerServiceMessages.h b/Engine/Source/Developer/ProfilerMessages/Classes/ProfilerServiceMessages.h
--- a/Engine/Source/Developer/ProfilerMessages/Classes/ProfilerServiceMessages.h
+++ b/Engine/Source/Developer/ProfilerMessages/Classes/ProfilerServiceMessages.h
@@ -12,30 +12,23 @@ struct FProfilerServiceAuthorize
 {
 	GENERATED_USTRUCT_BODY()
 
-	/**
-	 */
-	UPROPERTY()
+	/** */
+	UPROPERTY(EditAnywhere, Category="Message")
 	FGuid SessionId;
 
-	/**
-	 */
-	UPROPERTY()
+	/** */
+	UPROPERTY(EditAnywhere, Category="Message")
 	FGuid InstanceId;
 
-	/**
-	 */
-	UPROPERTY()
+	/** */
+	UPROPERTY(EditAnywhere, Category="Message")
 	TArray<uint8> Data;
 
-	/**
-	 * Default constructor.
-	 */
-	FProfilerServiceAuthorize( ) { }
+	/** Default constructor. */
+	FProfilerServiceAuthorize() { }
 
-	/**
-	 * Creates and initializes a new instance.
-	 */
-	FProfilerServiceAuthorize( const FGuid& InSessionId, const FGuid& InInstanceId, const TArray<uint8>& InData )
+	/** Creates and initializes a new instance. */
+	FProfilerServiceAuthorize(const FGuid& InSessionId, const FGuid& InInstanceId, const TArray<uint8>& InData)
 		: SessionId(InSessionId)
 		, InstanceId(InInstanceId)
 	{
@@ -52,41 +45,34 @@ struct FProfilerServiceData2
 {
 	GENERATED_USTRUCT_BODY()
 
-	/**
-	 */
-	UPROPERTY()
+	/** */
+	UPROPERTY(EditAnywhere, Category="Message")
 	FGuid InstanceId;
 
-	UPROPERTY()
+	UPROPERTY(EditAnywhere, Category="Message")
 	int64 Frame;
 
-	UPROPERTY()
+	UPROPERTY(EditAnywhere, Category="Message")
 	int32 CompressedSize;
 	
-	UPROPERTY()
+	UPROPERTY(EditAnywhere, Category="Message")
 	int32 UncompressedSize;
 
 	/** Profiler data encoded as string of hexes, cannot use TArray<uint8> because of the Message Bus limitation. */
-	UPROPERTY()
+	UPROPERTY(EditAnywhere, Category="Message")
 	FString HexData;
 
+	/** Default constructor. */
+	FProfilerServiceData2() { }
 
-	/**
-	 * Default constructor.
-	 */
-	FProfilerServiceData2( ) { }
-
-	/**
-	 * Creates and initializes a new instance.
-	 */
-	FProfilerServiceData2( const FGuid& InInstance, int64 InFrame, const FString& InHexData, int32 InCompressedSize, int32 InUncompressedSize )
-		: InstanceId( InInstance )
-		, Frame( InFrame )
-		, CompressedSize( InCompressedSize )
-		, UncompressedSize( InUncompressedSize )
-		, HexData( MoveTemp( InHexData ) )
-	{
-	}
+	/** Creates and initializes a new instance. */
+	FProfilerServiceData2(const FGuid& InInstance, int64 InFrame, const FString& InHexData, int32 InCompressedSize, int32 InUncompressedSize)
+		: InstanceId(InInstance)
+		, Frame(InFrame)
+		, CompressedSize(InCompressedSize)
+		, UncompressedSize(InUncompressedSize)
+		, HexData(MoveTemp(InHexData))
+	{ }
 };
 
 
@@ -97,23 +83,17 @@ struct FProfilerServicePreviewAck
 {
 	GENERATED_USTRUCT_BODY()
 
-	/**
-	 */
-	UPROPERTY()
+	/** */
+	UPROPERTY(EditAnywhere, Category="Message")
 	FGuid InstanceId;
 
-	/**
-	 * Default constructor.
-	 */
-	FProfilerServicePreviewAck( ) { }
+	/** Default constructor. */
+	FProfilerServicePreviewAck() { }
 
-	/**
-	 * Creates and initializes a new instance.
-	 */
-	FProfilerServicePreviewAck( const FGuid& InInstance )
+	/** Creates and initializes a new instance. */
+	FProfilerServicePreviewAck(const FGuid& InInstance)
 		: InstanceId(InInstance)
-	{
-	}
+	{ }
 };
 
 
@@ -127,55 +107,48 @@ struct FProfilerServiceFileChunk
 	GENERATED_USTRUCT_BODY()
 
 	/** The ID of the instance where this message should be sent. */
-	UPROPERTY()
+	UPROPERTY(EditAnywhere, Category="Message")
 	FGuid InstanceId;
 
 	/** The file containing this file chunk. */
-	UPROPERTY()
+	UPROPERTY(EditAnywhere, Category="Message")
 	FString Filename;
 
 	/** Data to be sent through message bus. Message bug doesn't support TArray<>, so we encode the data as HexString. */
-	UPROPERTY()
+	UPROPERTY(EditAnywhere, Category="Message")
 	FString HexData;
 
 	/** FProfilerFileChunkHeader stored in the array. */
-	UPROPERTY()
+	UPROPERTY(EditAnywhere, Category="Message")
 	TArray<uint8> Header;
 
 	/** Hash of this data and header. */
-	UPROPERTY()
+	UPROPERTY(EditAnywhere, Category="Message")
 	TArray<uint8> ChunkHash;
 
-	/**
-	 * Default constructor.
-	 */
-	FProfilerServiceFileChunk() 
-	{}
+	/** Default constructor. */
+	FProfilerServiceFileChunk()  { }
 
-	/**
-	 * Constructor for the new file chunk.
-	 */
+	/** Constructor for the new file chunk. */
 	FProfilerServiceFileChunk
 	( 
 		const FGuid& InInstanceID, 
 		const FString& InFilename, 
 		const TArray<uint8>& InHeader
 	)
-		: InstanceId( InInstanceID )
-		, Filename( InFilename )
-		, Header( InHeader )
-	{}
+		: InstanceId(InInstanceID)
+		, Filename(InFilename)
+		, Header(InHeader)
+	{ }
 
 	struct FNullTag{};
 
-	/**
-	 * Copy constructor, copies all properties, but not data.
-	 */
-	FProfilerServiceFileChunk( const FProfilerServiceFileChunk& ProfilerServiceFileChunk, FNullTag )
-		: InstanceId( ProfilerServiceFileChunk.InstanceId )
-		, Filename( ProfilerServiceFileChunk.Filename )
-		, Header( ProfilerServiceFileChunk.Header )
-	{}
+	/** Copy constructor, copies all properties, but not data. */
+	FProfilerServiceFileChunk(const FProfilerServiceFileChunk& ProfilerServiceFileChunk, FNullTag)
+		: InstanceId(ProfilerServiceFileChunk.InstanceId)
+		, Filename(ProfilerServiceFileChunk.Filename)
+		, Header(ProfilerServiceFileChunk.Header)
+	{ }
 };
 
 
@@ -205,25 +178,19 @@ struct FProfilerServiceSubscribe
 {
 	GENERATED_USTRUCT_BODY()
 
-	/**
-	 */
-	UPROPERTY()
+	/** */
+	UPROPERTY(EditAnywhere, Category="Message")
 	FGuid SessionId;
 
-	/**
-	 */
-	UPROPERTY()
+	/** */
+	UPROPERTY(EditAnywhere, Category="Message")
 	FGuid InstanceId;
 
-	/**
-	 * Default constructor.
-	 */
-	FProfilerServiceSubscribe( ) { }
+	/** Default constructor. */
+	FProfilerServiceSubscribe() { }
 
-	/**
-	 * Creates and initializes a new instance.
-	 */
-	FProfilerServiceSubscribe( const FGuid& InSessionId, const FGuid& InInstanceId )
+	/** Creates and initializes a new instance. */
+	FProfilerServiceSubscribe(const FGuid& InSessionId, const FGuid& InInstanceId)
 		: SessionId(InSessionId)
 		, InstanceId(InInstanceId)
 	{ }
@@ -237,26 +204,19 @@ struct FProfilerServiceUnsubscribe
 {
 	GENERATED_USTRUCT_BODY()
 
-	/**
-	 */
-	UPROPERTY()
+	/** */
+	UPROPERTY(EditAnywhere, Category="Message")
 	FGuid SessionId;
 
-	/**
-	 */
-	UPROPERTY()
+	/** */
+	UPROPERTY(EditAnywhere, Category="Message")
 	FGuid InstanceId;
 
+	/** Default constructor. */
+	FProfilerServiceUnsubscribe() { }
 
-	/**
-	 * Default constructor.
-	 */
-	FProfilerServiceUnsubscribe( ) { }
-
-	/**
-	 * Creates and initializes a new instance.
-	 */
-	FProfilerServiceUnsubscribe( const FGuid& InSessionId, const FGuid& InInstanceId )
+	/** Creates and initializes a new instance. */
+	FProfilerServiceUnsubscribe(const FGuid& InSessionId, const FGuid& InInstanceId)
 		: SessionId(InSessionId)
 		, InstanceId(InInstanceId)
 	{ }
@@ -270,22 +230,16 @@ struct FProfilerServiceCapture
 {
 	GENERATED_USTRUCT_BODY()
 
-	/**
-	 * The data capture state that should be set.
-	 */
-	UPROPERTY()
+	/** The data capture state that should be set. */
+	UPROPERTY(EditAnywhere, Category="Message")
 	bool bRequestedCaptureState;
 
-	/**
-	 * Default constructor.
-	 */
-	FProfilerServiceCapture( ) { }
+	/** Default constructor. */
+	FProfilerServiceCapture() { }
 
-	/**
-	 * Creates and initializes a new instance.
-	 */
-	FProfilerServiceCapture( const bool bInRequestedCaptureState )
-		: bRequestedCaptureState( bInRequestedCaptureState )
+	/** Creates and initializes a new instance. */
+	FProfilerServiceCapture(const bool bInRequestedCaptureState)
+		: bRequestedCaptureState(bInRequestedCaptureState)
 	{ }
 };
 
@@ -300,19 +254,15 @@ struct FProfilerServicePreview
 	/**
 	 * The data preview state that should be set.
 	 */
-	UPROPERTY()
+	UPROPERTY(EditAnywhere, Category="Message")
 	bool bRequestedPreviewState;
 
-	/**
-	 * Default constructor.
-	 */
-	FProfilerServicePreview( ) { }
+	/** Default constructor. */
+	FProfilerServicePreview() { }
 
-	/**
-	 * Creates and initializes a new instance.
-	 */
-	FProfilerServicePreview( const bool bInRequestedPreviewState )
-		: bRequestedPreviewState( bInRequestedPreviewState )
+	/** Creates and initializes a new instance. */
+	FProfilerServicePreview(const bool bInRequestedPreviewState)
+		: bRequestedPreviewState(bInRequestedPreviewState)
 	{ }
 };
 
@@ -325,18 +275,14 @@ struct FProfilerServiceRequest
 	GENERATED_USTRUCT_BODY()
 
 	/** Request @see EProfilerRequestType. */
-	UPROPERTY()
+	UPROPERTY(EditAnywhere, Category="Message")
 	uint32 Request;
 
-	/**
-	 * Default constructor.
-	 */
-	FProfilerServiceRequest( ) { }
+	/** Default constructor. */
+	FProfilerServiceRequest() { }
 
-	/**
-	 * Creates and initializes a new instance.
-	 */
-	FProfilerServiceRequest( uint32 InRequest )
+	/** Creates and initializes a new instance. */
+	FProfilerServiceRequest(uint32 InRequest)
 		: Request(InRequest)
 	{ }
 };
diff --git a/Engine/Source/Developer/SessionMessages/Classes/SessionServiceMessages.h b/Engine/Source/Developer/SessionMessages/Classes/SessionServiceMessages.h
--- a/Engine/Source/Developer/SessionMessages/Classes/SessionServiceMessages.h
+++ b/Engine/Source/Developer/SessionMessages/Classes/SessionServiceMessages.h
@@ -15,7 +15,6 @@ USTRUCT()
 struct FSessionServicePing
 {
 	GENERATED_USTRUCT_BODY()
-
 };
 
 
@@ -28,43 +27,43 @@ struct FSessionServicePong
 	GENERATED_USTRUCT_BODY()
 
 	/** Holds the application's build date. */
-	UPROPERTY()
+	UPROPERTY(EditAnywhere, Category="Message")
 	FString BuildDate;
 
 	/** Holds the name of the device that the application is running on. */
-	UPROPERTY()
+	UPROPERTY(EditAnywhere, Category="Message")
 	FString DeviceName;
 
 	/** Holds the application's instance identifier. */
-	UPROPERTY()
+	UPROPERTY(EditAnywhere, Category="Message")
 	FGuid InstanceId;
 
 	/** Holds the application's instance name. */
-	UPROPERTY()
+	UPROPERTY(EditAnywhere, Category="Message")
 	FString InstanceName;
 
 	/** Holds a flag indicating whether the application is running on a console. */
-	UPROPERTY()
+	UPROPERTY(EditAnywhere, Category="Message")
 	bool IsConsoleBuild;
 
 	/** Holds the name of the platform that the application is running on. */
-	UPROPERTY()
+	UPROPERTY(EditAnywhere, Category="Message")
 	FString PlatformName;
 
 	/** Holds the identifier of the session that the application belongs to. */
-	UPROPERTY()
+	UPROPERTY(EditAnywhere, Category="Message")
 	FGuid SessionId;
 
 	/** Holds the user defined name of the session. */
-	UPROPERTY()
+	UPROPERTY(EditAnywhere, Category="Message")
 	FString SessionName;
 
 	/** Holds the name of the user that started the session. */
-	UPROPERTY()
+	UPROPERTY(EditAnywhere, Category="Message")
 	FString SessionOwner;
 
 	/** Holds a flag indicating whether the application is the only one in that session. */
-	UPROPERTY()
+	UPROPERTY(EditAnywhere, Category="Message")
 	bool Standalone;
 };
 
@@ -81,23 +80,23 @@ struct FSessionServiceLog
 	GENERATED_USTRUCT_BODY()
 
 	/** Holds the log message category. */
-	UPROPERTY()
+	UPROPERTY(EditAnywhere, Category="Message")
 	FName Category;
 
 	/** Holds the log message data. */
-	UPROPERTY()
+	UPROPERTY(EditAnywhere, Category="Message")
 	FString Data;
 
 	/** Holds the application instance identifier. */
-	UPROPERTY()
+	UPROPERTY(EditAnywhere, Category="Message")
 	FGuid InstanceId;
 
 	/** Holds the time in seconds since the application was started. */
-	UPROPERTY()
+	UPROPERTY(EditAnywhere, Category="Message")
 	double TimeSeconds;
 
 	/** Holds the log message's verbosity level. */
-	UPROPERTY()
+	UPROPERTY(EditAnywhere, Category="Message")
 	uint8 Verbosity;
 
 public:
@@ -105,12 +104,12 @@ public:
 	/**
 	 * Default constructor.
 	 */
-	FSessionServiceLog( ) { }
+	FSessionServiceLog() { }
 
 	/**
 	 * Creates and initializes a new instance.
 	 */
-	FSessionServiceLog( const FName& InCategory, const FString& InData, const FGuid& InInstanceId, double InTimeSeconds, uint8 InVerbosity )
+	FSessionServiceLog(const FName& InCategory, const FString& InData, const FGuid& InInstanceId, double InTimeSeconds, uint8 InVerbosity)
 		: Category(InCategory)
 		, Data(InData)
 		, InstanceId(InInstanceId)
diff --git a/Engine/Source/Developer/TargetDeviceServices/Classes/TargetDeviceServiceMessages.h b/Engine/Source/Developer/TargetDeviceServices/Classes/TargetDeviceServiceMessages.h
--- a/Engine/Source/Developer/TargetDeviceServices/Classes/TargetDeviceServiceMessages.h
+++ b/Engine/Source/Developer/TargetDeviceServices/Classes/TargetDeviceServiceMessages.h
@@ -19,11 +19,11 @@ struct FTargetDeviceServiceDeployCommit
 	GENERATED_USTRUCT_BODY()
 
 	/** Holds the variant identifier of the target device to use. */
-	UPROPERTY()
+	UPROPERTY(EditAnywhere, Category="Message")
 	FName Variant;
 
 	/** Holds the identifier of the deployment transaction to commit. */
-	UPROPERTY()
+	UPROPERTY(EditAnywhere, Category="Message")
 	FGuid TransactionId;
 
 
@@ -51,11 +51,11 @@ struct FTargetDeviceServiceDeployFile
 	GENERATED_USTRUCT_BODY()
 
 	/** Holds the name and path of the file as it will be stored on the target device. */
-	UPROPERTY()
+	UPROPERTY(EditAnywhere, Category="Message")
 	FString TargetFileName;
 
 	/** Holds the identifier of the deployment transaction that this file is part of. */
-	UPROPERTY()
+	UPROPERTY(EditAnywhere, Category="Message")
 	FGuid TransactionId;
 
 
@@ -81,7 +81,7 @@ struct FTargetDeviceServiceDeployFinished
 	GENERATED_USTRUCT_BODY()
 
 	/** Holds the variant identifier of the target device to use. */
-	UPROPERTY()
+	UPROPERTY(EditAnywhere, Category="Message")
 	FName Variant;
 
 	/**
@@ -92,15 +92,15 @@ struct FTargetDeviceServiceDeployFinished
 	 * identifying the application. Application identifiers are returned from target device
 	 * services as result of successful deployment transactions.
 	 */
-	UPROPERTY()
+	UPROPERTY(EditAnywhere, Category="Message")
 	FString AppID;
 
 	/** Holds a flag indicating whether the deployment transaction finished successfully. */
-	UPROPERTY()
+	UPROPERTY(EditAnywhere, Category="Message")
 	bool Succeeded;
 
 	/** Holds the identifier of the deployment transaction that this file is part of. */
-	UPROPERTY()
+	UPROPERTY(EditAnywhere, Category="Message")
 	FGuid TransactionId;
 
 
@@ -133,7 +133,7 @@ struct FTargetDeviceServiceLaunchApp
 	GENERATED_USTRUCT_BODY()
 
 	/** Holds the variant identifier of the target device to use. */
-	UPROPERTY()
+	UPROPERTY(EditAnywhere, Category="Message")
 	FName Variant;
 
 	/**
@@ -144,15 +144,15 @@ struct FTargetDeviceServiceLaunchApp
 	 * identifying the application. Application identifiers are returned from target device
 	 * services as result of successful deployment transactions.
 	 */
-	UPROPERTY()
+	UPROPERTY(EditAnywhere, Category="Message")
 	FString AppID;
 
 	/** The application's build configuration, i.e. Debug or Shipping. */
-	UPROPERTY()
+	UPROPERTY(EditAnywhere, Category="Message")
 	uint8 BuildConfiguration;
 
 	/** Holds optional command line parameters for the application. */
-	UPROPERTY()
+	UPROPERTY(EditAnywhere, Category="Message")
 	FString Params;
 
 
@@ -187,15 +187,15 @@ struct FTargetDeviceServiceLaunchFinished
 	 * identifying the application. Application identifiers are returned from target device
 	 * services as result of successful deployment transactions.
 	 */
-	UPROPERTY()
+	UPROPERTY(EditAnywhere, Category="Message")
 	FString AppID;
 
 	/** Holds the process identifier for the launched application. */
-	UPROPERTY()
+	UPROPERTY(EditAnywhere, Category="Message")
 	int32 ProcessId;
 
 	/** Holds a flag indicating whether the application was launched successfully. */
-	UPROPERTY()
+	UPROPERTY(EditAnywhere, Category="Message")
 	bool Succeeded;
 
 
@@ -226,20 +226,20 @@ struct FTargetDeviceClaimDenied
 	GENERATED_USTRUCT_BODY()
 
 	/** Holds the identifier of the device that is already claimed. */
-	UPROPERTY()
+	UPROPERTY(EditAnywhere, Category="Message")
 	FString DeviceName;
 
 	/** Holds the name of the host computer that claimed the device. */
-	UPROPERTY()
+	UPROPERTY(EditAnywhere, Category="Message")
 	FString HostName;
 
 	/** Holds the name of the user that claimed the device. */
-	UPROPERTY()
+	UPROPERTY(EditAnywhere, Category="Message")
 	FString HostUser;
 
 
 	/** Default constructor. */
-	FTargetDeviceClaimDenied( ) { }
+	FTargetDeviceClaimDenied() { }
 
 	/** Creates and initializes a new instance. */
 	FTargetDeviceClaimDenied(const FString& InDeviceName, const FString& InHostName, const FString& InHostUser)
@@ -262,20 +262,20 @@ struct FTargetDeviceClaimed
 	GENERATED_USTRUCT_BODY()
 
 	/** Holds the identifier of the device that is being claimed. */
-	UPROPERTY()
+	UPROPERTY(EditAnywhere, Category="Message")
 	FString DeviceName;
 
 	/** Holds the name of the host computer that is claiming the device. */
-	UPROPERTY()
+	UPROPERTY(EditAnywhere, Category="Message")
 	FString HostName;
 
 	/** Holds the name of the user that is claiming the device. */
-	UPROPERTY()
+	UPROPERTY(EditAnywhere, Category="Message")
 	FString HostUser;
 
 
 	/** Default constructor. */
-	FTargetDeviceClaimed( ) { }
+	FTargetDeviceClaimed() { }
 
 	/** Creates and initializes a new instance. */
 	FTargetDeviceClaimed(const FString& InDeviceName, const FString& InHostName, const FString& InHostUser)
@@ -297,20 +297,20 @@ struct FTargetDeviceUnclaimed
 	GENERATED_USTRUCT_BODY()
 
 	/** Holds the identifier of the device that is no longer claimed. */
-	UPROPERTY()
+	UPROPERTY(EditAnywhere, Category="Message")
 	FString DeviceName;
 
 	/** Holds the name of the host computer that had claimed the device. */
-	UPROPERTY()
+	UPROPERTY(EditAnywhere, Category="Message")
 	FString HostName;
 
 	/** Holds the name of the user that had claimed the device. */
-	UPROPERTY()
+	UPROPERTY(EditAnywhere, Category="Message")
 	FString HostUser;
 
 
 	/** Default constructor. */
-	FTargetDeviceUnclaimed( ) { }
+	FTargetDeviceUnclaimed() { }
 
 	/** Creates and initializes a new instance. */
 	FTargetDeviceUnclaimed(const FString& InDeviceName, const FString& InHostName, const FString& InHostUser)
@@ -333,15 +333,15 @@ struct FTargetDeviceServicePing
 	GENERATED_USTRUCT_BODY()
 
 	/** Holds the name of the user who generated the ping. */
-	UPROPERTY()
+	UPROPERTY(EditAnywhere, Category="Message")
 	FString HostUser;
 
 
 	/** Default constructor. */
-	FTargetDeviceServicePing( ) { }
+	FTargetDeviceServicePing() { }
 
 	/** Creates and initializes a new instance. */
-	FTargetDeviceServicePing( const FString& InHostUser )
+	FTargetDeviceServicePing(const FString& InHostUser)
 		: HostUser(InHostUser)
 	{ }
 };
@@ -355,22 +355,22 @@ struct FTargetDeviceVariant
 {
 	GENERATED_USTRUCT_BODY()
 
-	UPROPERTY()
+	UPROPERTY(EditAnywhere, Category="Variant")
 	FString DeviceID;
 
-	UPROPERTY()
+	UPROPERTY(EditAnywhere, Category="Variant")
 	FName VariantName;
 
-	UPROPERTY()
+	UPROPERTY(EditAnywhere, Category="Variant")
 	FString TargetPlatformName;
 
-	UPROPERTY()
+	UPROPERTY(EditAnywhere, Category="Variant")
 	FName TargetPlatformId;
 
-	UPROPERTY()
+	UPROPERTY(EditAnywhere, Category="Variant")
 	FName VanillaPlatformId;
 
-	UPROPERTY()
+	UPROPERTY(EditAnywhere, Category="Variant")
 	FString PlatformDisplayName;
 };
 
@@ -384,71 +384,71 @@ struct FTargetDeviceServicePong
 	GENERATED_USTRUCT_BODY()
 
 	/** Holds a flag indicating whether the device is currently connected. */
-	UPROPERTY()
+	UPROPERTY(EditAnywhere, Category="Message")
 	bool Connected;
 
 	/** Holds the name of the host computer that the device is attached to. */
-	UPROPERTY()
+	UPROPERTY(EditAnywhere, Category="Message")
 	FString HostName;
 
 	/** Holds the name of the user under which the host computer is running. */
-	UPROPERTY()
+	UPROPERTY(EditAnywhere, Category="Message")
 	FString HostUser;
 
 	/** Holds the make of the device, i.e. Microsoft or Sony. */
-	UPROPERTY()
+	UPROPERTY(EditAnywhere, Category="Message")
 	FString Make;
 
 	/** Holds the model of the device. */
-	UPROPERTY()
+	UPROPERTY(EditAnywhere, Category="Message")
 	FString Model;
 
 	/** Holds the human readable name of the device, i.e "Bob's XBox'. */
-	UPROPERTY()
+	UPROPERTY(EditAnywhere, Category="Message")
 	FString Name;
 
 	/** Holds the name of the user that we log in to remote device as, i.e "root". */
-	UPROPERTY()
+	UPROPERTY(EditAnywhere, Category="Message")
 	FString DeviceUser;
 
 	/** Holds the password of the user that we log in to remote device as, i.e "12345". */
-	UPROPERTY()
+	UPROPERTY(EditAnywhere, Category="Message")
 	FString DeviceUserPassword;
 
 	/** Holds a flag indicating whether this device is shared with other users on the network. */
-	UPROPERTY()
+	UPROPERTY(EditAnywhere, Category="Message")
 	bool Shared;
 
 	/** Holds a flag indicating whether the device supports running multiple application instances in parallel. */
-	UPROPERTY()
+	UPROPERTY(EditAnywhere, Category="Message")
 	bool SupportsMultiLaunch;
 
 	/** Holds a flag indicating whether the device can be powered off. */
-	UPROPERTY()
+	UPROPERTY(EditAnywhere, Category="Message")
 	bool SupportsPowerOff;
 
 	/** Holds a flag indicating whether the device can be powered on. */
-	UPROPERTY()
+	UPROPERTY(EditAnywhere, Category="Message")
 	bool SupportsPowerOn;
 
 	/** Holds a flag indicating whether the device can be rebooted. */
-	UPROPERTY()
+	UPROPERTY(EditAnywhere, Category="Message")
 	bool SupportsReboot;
 
 	/** Holds a flag indicating whether the device's target platform supports variants. */
-	UPROPERTY()
+	UPROPERTY(EditAnywhere, Category="Message")
 	bool SupportsVariants;
 
 	/** Holds the device type. */
-	UPROPERTY()
+	UPROPERTY(EditAnywhere, Category="Message")
 	FString Type;
 
 	/** Holds the variant name of the default device. */
-	UPROPERTY()
+	UPROPERTY(EditAnywhere, Category="Message")
 	FName DefaultVariant;
 
 	/** List of the Flavors this service supports */
-	UPROPERTY()
+	UPROPERTY(EditAnywhere, Category="Message")
 	TArray<FTargetDeviceVariant> Variants;
 };
 
@@ -469,19 +469,19 @@ struct FTargetDeviceServicePowerOff
 	 *
 	 * If powering off is not enforced, if may fail if some running application prevents it.
 	 */
-	UPROPERTY()
+	UPROPERTY(EditAnywhere, Category="Message")
 	bool Force;
 
 	/** Holds the name of the user that wishes to power off the device. */
-	UPROPERTY()
+	UPROPERTY(EditAnywhere, Category="Message")
 	FString Operator;
 
 
 	/** Default constructor. */
-	FTargetDeviceServicePowerOff( ) { }
+	FTargetDeviceServicePowerOff() { }
 
 	/** Creates and initializes a new instance. */
-	FTargetDeviceServicePowerOff( const FString& InOperator, bool InForce )
+	FTargetDeviceServicePowerOff(const FString& InOperator, bool InForce)
 		: Force(InForce)
 		, Operator(InOperator)
 	{ }
@@ -497,15 +497,15 @@ struct FTargetDeviceServicePowerOn
 	GENERATED_USTRUCT_BODY()
 
 	/** Holds the name of the user that wishes to power on the device. */
-	UPROPERTY()
+	UPROPERTY(EditAnywhere, Category="Message")
 	FString Operator;
 
 
 	/** Default constructor. */
-	FTargetDeviceServicePowerOn( ) { }
+	FTargetDeviceServicePowerOn() { }
 
 	/** Creates and initializes a new instance. */
-	FTargetDeviceServicePowerOn( const FString& InOperator )
+	FTargetDeviceServicePowerOn(const FString& InOperator)
 		: Operator(InOperator)
 	{ }
 };
@@ -520,15 +520,15 @@ struct FTargetDeviceServiceReboot
 	GENERATED_USTRUCT_BODY()
 
 	/** Holds the name of the user that wishes to reboot the device. */
-	UPROPERTY()
+	UPROPERTY(EditAnywhere, Category="Message")
 	FString Operator;
 
 
 	/** Default constructor. */
-	FTargetDeviceServiceReboot( ) { }
+	FTargetDeviceServiceReboot() { }
 
 	/** Creates and initializes a new instance. */
-	FTargetDeviceServiceReboot( const FString& InOperator )
+	FTargetDeviceServiceReboot(const FString& InOperator)
 		: Operator(InOperator)
 	{ }
 };
@@ -547,15 +547,15 @@ struct FTargetDeviceServiceRunExecutable
 	GENERATED_USTRUCT_BODY()
 
 	/** Holds the variant identifier of the target device to use for execution. */
-	UPROPERTY()
+	UPROPERTY(EditAnywhere, Category="Message")
 	FName Variant;
 
 	/** Holds the path to the executable on the device. */
-	UPROPERTY()
+	UPROPERTY(EditAnywhere, Category="Message")
 	FString ExecutablePath;
 
 	/** Holds optional command line parameters for the executable. */
-	UPROPERTY()
+	UPROPERTY(EditAnywhere, Category="Message")
 	FString Params;
 
 
@@ -582,19 +582,19 @@ struct FTargetDeviceServiceRunFinished
 	GENERATED_USTRUCT_BODY()
 
 	/** Holds the variant identifier of the target device to use. */
-	UPROPERTY()
+	UPROPERTY(EditAnywhere, Category="Message")
 	FName Variant;
 
 	/** Holds the path to the executable that was run. */
-	UPROPERTY()
+	UPROPERTY(EditAnywhere, Category="Message")
 	FString ExecutablePath;
 
 	/** Holds the process identifier of the running executable. */
-	UPROPERTY()
+	UPROPERTY(EditAnywhere, Category="Message")
 	int32 ProcessId;
 
 	/** Holds a flag indicating whether the executable started successfully. */
-	UPROPERTY()
+	UPROPERTY(EditAnywhere, Category="Message")
 	bool Succeeded;
 
 
diff --git a/Engine/Source/Editor/SwarmInterface/Private/SwarmMessages.h b/Engine/Source/Editor/SwarmInterface/Private/SwarmMessages.h
--- a/Engine/Source/Editor/SwarmInterface/Private/SwarmMessages.h
+++ b/Engine/Source/Editor/SwarmInterface/Private/SwarmMessages.h
@@ -18,19 +18,18 @@ struct FSwarmPongMessage
 {
 	GENERATED_USTRUCT_BODY()
 
-	UPROPERTY()
-	bool			bIsEditor;
+	UPROPERTY(EditAnywhere, Category="Message")
+	bool bIsEditor;
 
-	UPROPERTY()
-	FString			ComputerName;
+	UPROPERTY(EditAnywhere, Category="Message")
+	FString ComputerName;
 
-	FSwarmPongMessage() {}
+	FSwarmPongMessage() { }
 
 	FSwarmPongMessage(bool bInIsEditor, const TCHAR* InComputerName)
-	:	bIsEditor(bInIsEditor)
-	,	ComputerName(InComputerName)
-	{
-	}
+		: bIsEditor(bInIsEditor)
+		, ComputerName(InComputerName)
+	{ }
 };
 
 
@@ -39,15 +38,14 @@ struct FSwarmInfoMessage
 {
 	GENERATED_USTRUCT_BODY()
 
-	UPROPERTY()
+	UPROPERTY(EditAnywhere, Category="Message")
 	FString TextMessage;
 
-	FSwarmInfoMessage() {}
+	FSwarmInfoMessage() { }
 
 	FSwarmInfoMessage(const FString& InTextMessage)
-		:	TextMessage(InTextMessage)
-	{
-	}
+		: TextMessage(InTextMessage)
+	{ }
 };
 
 
@@ -57,31 +55,34 @@ struct FSwarmAlertMessage
 	GENERATED_USTRUCT_BODY()
 
 	/** The Job Guid */
-	UPROPERTY()
-	FGuid			JobGuid;
+	UPROPERTY(EditAnywhere, Category="Message")
+	FGuid JobGuid;
+
 	/** The type of alert */
-	UPROPERTY()
-	uint8			AlertLevel;
+	UPROPERTY(EditAnywhere, Category="Message")
+	uint8 AlertLevel;
+	
 	/** The identifier for the object that is associated with the issue */
-	UPROPERTY()
-	FGuid			ObjectGuid;
+	UPROPERTY(EditAnywhere, Category="Message")
+	FGuid ObjectGuid;
+	
 	/** App-specific identifier for the type of the object */
-	UPROPERTY()
-	int32			TypeId;
+	UPROPERTY(EditAnywhere, Category="Message")
+	int32 TypeId;
+	
 	/** Generic text message for informational purposes */
-	UPROPERTY()
-	FString			TextMessage;
+	UPROPERTY(EditAnywhere, Category="Message")
+	FString TextMessage;
 
-	FSwarmAlertMessage() {}
+	FSwarmAlertMessage() { }
 
 	FSwarmAlertMessage(const FGuid& InJobGuid, NSwarm::TAlertLevel InAlertLevel, const FGuid& InObjectGuid, int32 InTypeId, const FString& InTextMessage)
-		:	JobGuid(InJobGuid)
-		,	AlertLevel(InAlertLevel)
-		,	ObjectGuid(InObjectGuid)
-		,	TypeId(InTypeId)
-		,	TextMessage(InTextMessage)
-	{
-	}
+		: JobGuid(InJobGuid)
+		, AlertLevel(InAlertLevel)
+		, ObjectGuid(InObjectGuid)
+		, TypeId(InTypeId)
+		, TextMessage(InTextMessage)
+	{ }
 };
 
 
@@ -91,19 +92,19 @@ struct FSwarmTimingMessage
 	GENERATED_USTRUCT_BODY()
 
 	/** State that the distributed job is transitioning to */
-	UPROPERTY()
-	uint8			State;
+	UPROPERTY(EditAnywhere, Category="Message")
+	uint8 State;
+
 	/** The thread this state is referring to */
-	UPROPERTY()
-	int32			ThreadNum;
+	UPROPERTY(EditAnywhere, Category="Message")
+	int32 ThreadNum;
 
-	FSwarmTimingMessage() {}
+	FSwarmTimingMessage() { }
 
 	FSwarmTimingMessage(NSwarm::TProgressionState InState, int32 InThreadNum)
-		:	State(InState)
-		,	ThreadNum(InThreadNum)
-	{
-	}
+		: State(InState)
+		, ThreadNum(InThreadNum)
+	{ }
 };
 
 
@@ -112,7 +113,7 @@ struct FSwarmTaskRequestMessage
 {
 	GENERATED_USTRUCT_BODY()
 
-	FSwarmTaskRequestMessage() {}
+	FSwarmTaskRequestMessage() { }
 };
 
 
@@ -121,7 +122,7 @@ struct FSwarmTaskRequestReleaseMessage
 {
 	GENERATED_USTRUCT_BODY()
 
-	FSwarmTaskRequestReleaseMessage() {}
+	FSwarmTaskRequestReleaseMessage() { }
 };
 
 
@@ -130,7 +131,7 @@ struct FSwarmTaskRequestReservationMessage
 {
 	GENERATED_USTRUCT_BODY()
 
-	FSwarmTaskRequestReservationMessage() {}
+	FSwarmTaskRequestReservationMessage() { }
 };
 
 
@@ -140,28 +141,32 @@ struct FSwarmTaskRequestSpecificationMessage
 	GENERATED_USTRUCT_BODY()
 
 	/** The GUID used for identifying the Task being referred to */
-	UPROPERTY()
-	FGuid					TaskGuid;
+	UPROPERTY(EditAnywhere, Category="Message")
+	FGuid TaskGuid;
+
 	/** The Task's parameter string specified with AddTask */
-	UPROPERTY()
-	FString					Parameters;
+	UPROPERTY(EditAnywhere, Category="Message")
+	FString Parameters;
+
 	/** Flags used to control the behavior of the Task, subject to overrides from the containing Job */
-	UPROPERTY()
-	uint8					Flags;
+	UPROPERTY(EditAnywhere, Category="Message")
+	uint8 Flags;
+
 	/** The Task's cost, relative to all other Tasks in the same Job, used for even distribution and scheduling */
-	UPROPERTY()
-	uint32					Cost;
+	UPROPERTY(EditAnywhere, Category="Message")
+	uint32 Cost;
+
 	/** Any additional Task dependencies */
-	UPROPERTY()
-	TArray<FString>			Dependencies;
+	UPROPERTY(EditAnywhere, Category="Message")
+	TArray<FString> Dependencies;
 
-	FSwarmTaskRequestSpecificationMessage() {}
+	FSwarmTaskRequestSpecificationMessage() { }
 
 	FSwarmTaskRequestSpecificationMessage(FGuid InTaskGuid, const FString& InParameters, NSwarm::TJobTaskFlags InFlags, const TArray<FString>& InDependencies)
-		:	TaskGuid(InTaskGuid)
-		,	Parameters(InParameters)
-		,	Flags(InFlags)
-		,	Dependencies(InDependencies)
+		: TaskGuid(InTaskGuid)
+		, Parameters(InParameters)
+		, Flags(InFlags)
+		, Dependencies(InDependencies)
 	{ }
 };
 
@@ -172,29 +177,32 @@ struct FSwarmJobStateMessage
 	GENERATED_USTRUCT_BODY()
 
 	/** The Job GUID used for identifying the Job */
-	UPROPERTY()
-	FGuid			Guid;
+	UPROPERTY(EditAnywhere, Category="Message")
+	FGuid Guid;
+
 	/** The current state and arbitrary message */
-	UPROPERTY()
-	uint8			State;
-	UPROPERTY()
-	FString			Message;
+	UPROPERTY(EditAnywhere, Category="Message")
+	uint8 State;
+
+	UPROPERTY(EditAnywhere, Category="Message")
+	FString Message;
+
 	/** Various stats, including run time, exit codes, etc. */
-	UPROPERTY()
-	int32			ExitCode;
-	UPROPERTY()
-	double			RunningTime;
+	UPROPERTY(EditAnywhere, Category="Message")
+	int32 ExitCode;
 
-	FSwarmJobStateMessage() {}
+	UPROPERTY(EditAnywhere, Category="Message")
+	double RunningTime;
+
+	FSwarmJobStateMessage() { }
 
 	FSwarmJobStateMessage(FGuid InGuid, NSwarm::TJobTaskState InState, const FString& InMessage, int32 InExitCode, double InRunningTime)
-		:	Guid(InGuid)
-		,	State(InState)
-		,	Message(InMessage)
-		,	ExitCode(InExitCode)
-		,	RunningTime(InRunningTime)
-	{
-	}
+		: Guid(InGuid)
+		, State(InState)
+		, Message(InMessage)
+		, ExitCode(InExitCode)
+		, RunningTime(InRunningTime)
+	{ }
 };
 
 
@@ -204,29 +212,32 @@ struct FSwarmTaskStateMessage
 	GENERATED_USTRUCT_BODY()
 
 	/** The Task GUID used for identifying the Task */
-	UPROPERTY()
-	FGuid			Guid;
+	UPROPERTY(EditAnywhere, Category="Message")
+	FGuid Guid;
+
 	/** The current state and arbitrary message */
-	UPROPERTY()
-	uint8			State;
-	UPROPERTY()
-	FString			Message;
+	UPROPERTY(EditAnywhere, Category="Message")
+	uint8 State;
+
+	UPROPERTY(EditAnywhere, Category="Message")
+	FString Message;
+
 	/** Various stats, including run time, exit codes, etc. */
-	UPROPERTY()
-	int32			ExitCode;
-	UPROPERTY()
-	double			RunningTime;
+	UPROPERTY(EditAnywhere, Category="Message")
+	int32 ExitCode;
 
-	FSwarmTaskStateMessage() {}
+	UPROPERTY(EditAnywhere, Category="Message")
+	double RunningTime;
+
+	FSwarmTaskStateMessage() { }
 
 	FSwarmTaskStateMessage(FGuid InGuid, NSwarm::TJobTaskState InState, const FString& InMessage, int32 InExitCode, double InRunningTime)
-		:	Guid(InGuid)
-		,	State(InState)
-		,	Message(InMessage)
-		,	ExitCode(InExitCode)
-		,	RunningTime(InRunningTime)
-	{
-	}
+		: Guid(InGuid)
+		, State(InState)
+		, Message(InMessage)
+		, ExitCode(InExitCode)
+		, RunningTime(InRunningTime)
+	{ }
 };
 
 
@@ -235,5 +246,5 @@ struct FSwarmQuitMessage
 {
 	GENERATED_USTRUCT_BODY()
 
-	FSwarmQuitMessage() {}
+	FSwarmQuitMessage() { }
 };
diff --git a/Engine/Source/Editor/UnrealEdMessages/Classes/AssetEditorMessages.h b/Engine/Source/Editor/UnrealEdMessages/Classes/AssetEditorMessages.h
--- a/Engine/Source/Editor/UnrealEdMessages/Classes/AssetEditorMessages.h
+++ b/Engine/Source/Editor/UnrealEdMessages/Classes/AssetEditorMessages.h
@@ -16,19 +16,15 @@ struct FAssetEditorRequestOpenAsset
 	/**
 	 * Holds the name of the asset to open.
 	 */
-	UPROPERTY()
+	UPROPERTY(EditAnywhere, Category="Message")
 	FString AssetName;
 
 
-	/**
-	 * Default constructor.
-	 */
-	FAssetEditorRequestOpenAsset( ) { }
+	/** Default constructor. */
+	FAssetEditorRequestOpenAsset() { }
 
-	/**
-	 * Creates and initializes a new instance.
-	 */
-	FAssetEditorRequestOpenAsset( const FString& InAssetName )
+	/** Creates and initializes a new instance. */
+	FAssetEditorRequestOpenAsset(const FString& InAssetName)
 		: AssetName(InAssetName)
 	{ }
 };
diff --git a/Engine/Source/Editor/UnrealEdMessages/Classes/FileServerMessages.h b/Engine/Source/Editor/UnrealEdMessages/Classes/FileServerMessages.h
--- a/Engine/Source/Editor/UnrealEdMessages/Classes/FileServerMessages.h
+++ b/Engine/Source/Editor/UnrealEdMessages/Classes/FileServerMessages.h
@@ -13,28 +13,19 @@ struct FFileServerReady
 {
 	GENERATED_USTRUCT_BODY()
 	
-	/**
-	 * Holds the list of IP addresses that the file server is listening on.
-	 */
-	UPROPERTY()
+	/** Holds the list of IP addresses that the file server is listening on. */
+	UPROPERTY(EditAnywhere, Category="Message")
 	TArray<FString> AddressList;
 	
-	/**
-	 * Holds the file server's application identifier.
-	 */
-	UPROPERTY()
+	/** Holds the file server's application identifier. */
+	UPROPERTY(EditAnywhere, Category="Message")
 	FGuid InstanceId;
 
+	/** Default constructor. */
+	FFileServerReady() { }
 
-	/**
-	 * Default constructor.
-	 */
-	FFileServerReady( ) { }
-
-	/**
-	 * Creates and initializes a new instance.
-	 */
-	FFileServerReady( const TArray<FString>& InAddressList, const FGuid& InInstanceId )
+	/** Creates and initializes a new instance. */
+	FFileServerReady(const TArray<FString>& InAddressList, const FGuid& InInstanceId)
 		: AddressList(InAddressList)
 		, InstanceId(InInstanceId)
 	{ }
diff --git a/Engine/Source/Runtime/AutomationMessages/Classes/AutomationWorkerMessages.h b/Engine/Source/Runtime/AutomationMessages/Classes/AutomationWorkerMessages.h
--- a/Engine/Source/Runtime/AutomationMessages/Classes/AutomationWorkerMessages.h
+++ b/Engine/Source/Runtime/AutomationMessages/Classes/AutomationWorkerMessages.h
@@ -17,26 +17,26 @@ struct FAutomationWorkerFindWorkers
 	GENERATED_USTRUCT_BODY()
 
 	/** Holds the change list number to find workers for. */
-	UPROPERTY()
+	UPROPERTY(EditAnywhere, Category="Message")
 	int32 Changelist;
 
 	/** The name of the game. */
-	UPROPERTY()
+	UPROPERTY(EditAnywhere, Category="Message")
 	FString GameName;
 
 	/** The name of the process. */
-	UPROPERTY()
+	UPROPERTY(EditAnywhere, Category="Message")
 	FString ProcessName;
 
 	/** Holds the session identifier to find workers for. */
-	UPROPERTY()
+	UPROPERTY(EditAnywhere, Category="Message")
 	FGuid SessionId;
 
 	/** Default constructor. */
-	FAutomationWorkerFindWorkers( ) { }
+	FAutomationWorkerFindWorkers() { }
 
 	/** Creates and initializes a new instance. */
-	FAutomationWorkerFindWorkers( int32 InChangelist, const FString& InGameName, const FString& InProcessName, const FGuid& InSessionId )
+	FAutomationWorkerFindWorkers(int32 InChangelist, const FString& InGameName, const FString& InProcessName, const FGuid& InSessionId)
 		: Changelist(InChangelist)
 		, GameName(InGameName)
 		, ProcessName(InProcessName)
@@ -54,49 +54,47 @@ struct FAutomationWorkerFindWorkersResponse
 	GENERATED_USTRUCT_BODY()
 
 	/** Holds the name of the device that the worker is running on. */
-	UPROPERTY()
+	UPROPERTY(EditAnywhere, Category="Message")
 	FString DeviceName;
 
 	/** Holds the name of the worker's application instance. */
-	UPROPERTY()
+	UPROPERTY(EditAnywhere, Category="Message")
 	FString InstanceName;
 
 	/** Holds the name of the platform that the worker is running on. */
-	UPROPERTY()
+	UPROPERTY(EditAnywhere, Category="Message")
 	FString Platform;
 
 	/** Holds the name of the operating system version. */
-	UPROPERTY()
+	UPROPERTY(EditAnywhere, Category="Message")
 	FString OSVersionName;
 
 	/** Holds the name of the device model. */
-	UPROPERTY()
+	UPROPERTY(EditAnywhere, Category="Message")
 	FString ModelName;
 
 	/** Holds the name of the GPU. */
-	UPROPERTY()
+	UPROPERTY(EditAnywhere, Category="Message")
 	FString GPUName;
 
 	/** Holds the name of the CPU model. */
-	UPROPERTY()
+	UPROPERTY(EditAnywhere, Category="Message")
 	FString CPUModelName;
 
 	/** Holds the amount of RAM this device has in gigabytes. */
-	UPROPERTY()
+	UPROPERTY(EditAnywhere, Category="Message")
 	uint32 RAMInGB;
 
 	/** Holds the name of the current render mode. */
-	UPROPERTY()
+	UPROPERTY(EditAnywhere, Category="Message")
 	FString RenderModeName;
 
 	/** Holds the worker's application session identifier. */
-	UPROPERTY()
+	UPROPERTY(EditAnywhere, Category="Message")
 	FGuid SessionId;
 
-	/**
-	 * Default constructor.
-	 */
-	FAutomationWorkerFindWorkersResponse( ) { }
+	/** Default constructor. */
+	FAutomationWorkerFindWorkersResponse() { }
 };
 
 
@@ -146,21 +144,17 @@ struct FAutomationWorkerRequestTests
 	GENERATED_USTRUCT_BODY()
 
 	/** Holds a flag indicating whether the developer directory should be included. */
-	UPROPERTY()
+	UPROPERTY(EditAnywhere, Category="Message")
 	bool DeveloperDirectoryIncluded;
 
 	/** Holds a flag indicating which tests we'd like to request. */
-	UPROPERTY()
+	UPROPERTY(EditAnywhere, Category="Message")
 	uint32 RequestedTestFlags;
 
-	/**
-	 * Default constructor.
-	 */
-	FAutomationWorkerRequestTests( ) { }
+	/** Default constructor. */
+	FAutomationWorkerRequestTests() { }
 
-	/**
-	 * Creates and initializes a new instance.
-	 */
+	/** Creates and initializes a new instance. */
 	FAutomationWorkerRequestTests(bool InDeveloperDirectoryIncluded, uint32 InRequestedTestFlags)
 		: DeveloperDirectoryIncluded(InDeveloperDirectoryIncluded)
 		, RequestedTestFlags(InRequestedTestFlags)
@@ -177,22 +171,18 @@ struct FAutomationWorkerRequestTestsReply
 	GENERATED_USTRUCT_BODY()
 
 	/** Holds the test information serialized into a string. */
-	UPROPERTY()
+	UPROPERTY(EditAnywhere, Category="Message")
 	FString TestInfo;
 
 	/** Holds the total number of tests returned. */
-	UPROPERTY()
+	UPROPERTY(EditAnywhere, Category="Message")
 	int32 TotalNumTests;
 
-	/**
-	 * Default constructor.
-	 */
-	FAutomationWorkerRequestTestsReply( ) { }
+	/** Default constructor. */
+	FAutomationWorkerRequestTestsReply() { }
 
-	/**
-	 * Creates and initializes a new instance.
-	 */
-	FAutomationWorkerRequestTestsReply( const FString& InTestInfo, const int32& InTotalNumTests )
+	/** Creates and initializes a new instance. */
+	FAutomationWorkerRequestTestsReply(const FString& InTestInfo, const int32& InTotalNumTests)
 		: TestInfo(InTestInfo)
 		, TotalNumTests(InTotalNumTests)
 	{ }
@@ -207,6 +197,7 @@ struct FAutomationWorkerRequestTestsReplyComplete
 	GENERATED_USTRUCT_BODY()
 };
 
+
 /**
  * Implements a message to request the running of automation tests on a worker.
  */
@@ -216,34 +207,30 @@ struct FAutomationWorkerRunTests
 	GENERATED_USTRUCT_BODY()
 
 	/** */
-	UPROPERTY()
+	UPROPERTY(EditAnywhere, Category="Message")
 	uint32 ExecutionCount;
 
 	/** */
-	UPROPERTY()
+	UPROPERTY(EditAnywhere, Category="Message")
 	int32 RoleIndex;
 
 	/** Holds the name of the test to run. */
-	UPROPERTY()
+	UPROPERTY(EditAnywhere, Category="Message")
 	FString TestName;
 
 	/** If true, we will save out screenshots for tests that support them. */
-	UPROPERTY()
+	UPROPERTY(EditAnywhere, Category="Message")
 	bool bScreenshotsEnabled;
 
 	/** If true, we will not resize screen shots. */
-	UPROPERTY()
+	UPROPERTY(EditAnywhere, Category="Message")
 	bool bUseFullSizeScreenShots;
 
-	/**
-	 * Default constructor.
-	 */
+	/** Default constructor. */
 	FAutomationWorkerRunTests( ) { }
 
-	/**
-	 * Creates and initializes a new instance.
-	 */
-	FAutomationWorkerRunTests( uint32 InExecutionCount, int32 InRoleIndex, FString InTestName, bool InScreenshotsEnabled, bool InFullSizeScreenshots )
+	/** Creates and initializes a new instance. */
+	FAutomationWorkerRunTests(uint32 InExecutionCount, int32 InRoleIndex, FString InTestName, bool InScreenshotsEnabled, bool InFullSizeScreenshots)
 		: ExecutionCount(InExecutionCount)
 		, RoleIndex(InRoleIndex)
 		, TestName(InTestName)
@@ -262,31 +249,31 @@ struct FAutomationWorkerRunTestsReply
 	GENERATED_USTRUCT_BODY()
 
 	/** */
-	UPROPERTY()
+	UPROPERTY(EditAnywhere, Category="Message")
 	float Duration;
 
 	/** */
-	UPROPERTY()
+	UPROPERTY(EditAnywhere, Category="Message")
 	TArray<FString> Errors;
 
 	/** */
-	UPROPERTY()
+	UPROPERTY(EditAnywhere, Category="Message")
 	uint32 ExecutionCount;
 
 	/** */
-	UPROPERTY()
+	UPROPERTY(EditAnywhere, Category="Message")
 	TArray<FString> Logs;
 
 	/** */
-	UPROPERTY()
+	UPROPERTY(EditAnywhere, Category="Message")
 	bool Success;
 
 	/** */
-	UPROPERTY()
+	UPROPERTY(EditAnywhere, Category="Message")
 	FString TestName;
 
 	/** */
-	UPROPERTY()
+	UPROPERTY(EditAnywhere, Category="Message")
 	TArray<FString> Warnings;
 };
 
@@ -299,18 +286,14 @@ struct FAutomationWorkerRequestNextNetworkCommand
 	GENERATED_USTRUCT_BODY()
 
 	/** */
-	UPROPERTY()
+	UPROPERTY(EditAnywhere, Category="Message")
 	uint32 ExecutionCount;
 
-	/**
-	 * Default constructor.
-	 */
-	FAutomationWorkerRequestNextNetworkCommand( ) { }
+	/** Default constructor. */
+	FAutomationWorkerRequestNextNetworkCommand() { }
 
-	/**
-	 * Creates and initializes a new instance.
-	 */
-	FAutomationWorkerRequestNextNetworkCommand( uint32 InExecutionCount )
+	/** Creates and initializes a new instance. */
+	FAutomationWorkerRequestNextNetworkCommand(uint32 InExecutionCount)
 		: ExecutionCount(InExecutionCount)
 	{ }
 };
@@ -334,10 +317,10 @@ struct FAutomationWorkerScreenImage
 	GENERATED_USTRUCT_BODY()
 
 	/** The screen shot data. */
-	UPROPERTY()
+	UPROPERTY(EditAnywhere, Category="Message")
 	TArray<uint8> ScreenImage;
 
 	/** The screen shot name. */
-	UPROPERTY()
+	UPROPERTY(EditAnywhere, Category="Message")
 	FString ScreenShotName;
 };
diff --git a/Engine/Source/Runtime/EngineMessages/Classes/EngineServiceMessages.h b/Engine/Source/Runtime/EngineMessages/Classes/EngineServiceMessages.h
--- a/Engine/Source/Runtime/EngineMessages/Classes/EngineServiceMessages.h
+++ b/Engine/Source/Runtime/EngineMessages/Classes/EngineServiceMessages.h
@@ -27,31 +27,31 @@ struct FEngineServicePong
 	GENERATED_USTRUCT_BODY()
 
 	/** Holds the name of the currently loaded level, if any. */
-	UPROPERTY()
+	UPROPERTY(EditAnywhere, Category="Message")
 	FString CurrentLevel;
 
 	/** Holds the engine version. */
-	UPROPERTY()
+	UPROPERTY(EditAnywhere, Category="Message")
 	int32 EngineVersion;
 
 	/** Holds a flag indicating whether game play has begun. */
-	UPROPERTY()
+	UPROPERTY(EditAnywhere, Category="Message")
 	bool HasBegunPlay;
 
 	/** Holds the instance identifier. */
-	UPROPERTY()
+	UPROPERTY(EditAnywhere, Category="Message")
 	FGuid InstanceId;
 
 	/** Holds the type of the engine instance. */
-	UPROPERTY()
+	UPROPERTY(EditAnywhere, Category="Message")
 	FString InstanceType;
 
 	/** Holds the identifier of the session that the application belongs to. */
-	UPROPERTY()
+	UPROPERTY(EditAnywhere, Category="Message")
 	FGuid SessionId;
 
 	/** Holds the time in seconds since the world was loaded. */
-	UPROPERTY()
+	UPROPERTY(EditAnywhere, Category="Message")
 	float WorldTimeSeconds;
 };
 
@@ -68,11 +68,11 @@ struct FEngineServiceAuthDeny
 	GENERATED_USTRUCT_BODY()
 
 	/** Holds the name of the user that denied access. */
-	UPROPERTY()
+	UPROPERTY(EditAnywhere, Category="Message")
 	FString UserName;
 
 	/** Holds the name of the user that access is denied to. */
-	UPROPERTY()
+	UPROPERTY(EditAnywhere, Category="Message")
 	FString UserToDeny;
 };
 
@@ -86,11 +86,11 @@ struct FEngineServiceAuthGrant
 	GENERATED_USTRUCT_BODY()
 
 	/** Holds the name of the user that granted access. */
-	UPROPERTY()
+	UPROPERTY(EditAnywhere, Category="Message")
 	FString UserName;
 
 	/** Holds the name of the user that access is granted to. */
-	UPROPERTY()
+	UPROPERTY(EditAnywhere, Category="Message")
 	FString UserToGrant;
 
 };
@@ -108,22 +108,18 @@ struct FEngineServiceExecuteCommand
 	GENERATED_USTRUCT_BODY()
 
 	/** Holds the command to execute. */
-	UPROPERTY()
+	UPROPERTY(EditAnywhere, Category="Message")
 	FString Command;
 
 	/** Holds the name of the user that wants to execute the command. */
-	UPROPERTY()
+	UPROPERTY(EditAnywhere, Category="Message")
 	FString UserName;
 
-	/**
-	 * Default constructor.
-	 */
-	FEngineServiceExecuteCommand( ) { }
+	/** Default constructor. */
+	FEngineServiceExecuteCommand() { }
 
-	/**
-	 * Creates and initializes a new instance.
-	 */
-	FEngineServiceExecuteCommand( const FString& InCommand, const FString& InUserName )
+	/** Creates and initializes a new instance. */
+	FEngineServiceExecuteCommand(const FString& InCommand, const FString& InUserName)
 		: Command(InCommand)
 		, UserName(InUserName)
 	{ }
@@ -139,18 +135,14 @@ struct FEngineServiceTerminate
 	GENERATED_USTRUCT_BODY()
 
 	/** Holds the name of the user that wants to terminate the engine. */
-	UPROPERTY()
+	UPROPERTY(EditAnywhere, Category="Message")
 	FString UserName;
 
-	/**
-	 * Default constructor.
-	 */
-	FEngineServiceTerminate( ) { }
+	/** Default constructor. */
+	FEngineServiceTerminate() { }
 
-	/**
-	 * Creates and initializes a new instance.
-	 */
-	FEngineServiceTerminate( const FString& InUserName )
+	/** Creates and initializes a new instance. */
+	FEngineServiceTerminate(const FString& InUserName)
 		: UserName(InUserName)
 	{ }
 };
@@ -168,22 +160,18 @@ struct FEngineServiceNotification
 	GENERATED_USTRUCT_BODY()
 
 	/** Holds the notification text. */
-	UPROPERTY()
+	UPROPERTY(EditAnywhere, Category="Message")
 	FString Text;
 
 	/** Holds the time in seconds since the engine started. */
-	UPROPERTY()
+	UPROPERTY(EditAnywhere, Category="Message")
 	double TimeSeconds;
 
-	/**
-	 * Default constructor.
-	 */
-	FEngineServiceNotification( ) { }
+	/** Default constructor. */
+	FEngineServiceNotification() { }
 
-	/**
-	 * Creates and initializes a new instance.
-	 */
-	FEngineServiceNotification( const FString& InText, double InTimeSeconds )
+	/** Creates and initializes a new instance. */
+	FEngineServiceNotification(const FString& InText, double InTimeSeconds)
 		: Text(InText)
 		, TimeSeconds(InTimeSeconds)
 	{ }
diff --git a/Engine/Source/Runtime/MessagingRpc/Public/MessageRpcMessages.h b/Engine/Source/Runtime/MessagingRpc/Public/MessageRpcMessages.h
--- a/Engine/Source/Runtime/MessagingRpc/Public/MessageRpcMessages.h
+++ b/Engine/Source/Runtime/MessagingRpc/Public/MessageRpcMessages.h
@@ -12,7 +12,7 @@ struct FMessageRpcCancel
 	GENERATED_USTRUCT_BODY()
 
 	/** Correlation identifier for the RPC call that this message refers to. */
-	UPROPERTY()
+	UPROPERTY(EditAnywhere, Category="Message")
 	FGuid CallId;
 
 	/** Default constructor. */
@@ -32,15 +32,15 @@ struct FMessageRpcProgress
 	GENERATED_USTRUCT_BODY()
 
 	/** Completion percentage (0.0 to 1.0). */
-	UPROPERTY()
+	UPROPERTY(EditAnywhere, Category="Message")
 	float Completion;
 
 	/** Correlation identifier for the RPC call that this message refers to. */
-	UPROPERTY()
+	UPROPERTY(EditAnywhere, Category="Message")
 	FGuid CallId;
 
 	/** Status text. */
-	UPROPERTY()
+	UPROPERTY(EditAnywhere, Category="Message")
 	FString StatusText;
 
 	/** Default constructor. */
@@ -62,7 +62,7 @@ struct FMessageRpcUnhandled
 	GENERATED_USTRUCT_BODY()
 
 	/** Correlation identifier for the RPC call that this message refers to. */
-	UPROPERTY()
+	UPROPERTY(EditAnywhere, Category="Message")
 	FGuid CallId;
 
 	/** Default constructor. */
diff --git a/Engine/Source/Runtime/Portal/Messages/Public/PortalApplicationWindowMessages.h b/Engine/Source/Runtime/Portal/Messages/Public/PortalApplicationWindowMessages.h
--- a/Engine/Source/Runtime/Portal/Messages/Public/PortalApplicationWindowMessages.h
+++ b/Engine/Source/Runtime/Portal/Messages/Public/PortalApplicationWindowMessages.h
@@ -10,7 +10,7 @@ struct FPortalApplicationWindowNavigateToRequest
 {
 	GENERATED_USTRUCT_BODY()
 
-	UPROPERTY()
+	UPROPERTY(EditAnywhere, Category="Message")
 	FString Url;
 
 	FPortalApplicationWindowNavigateToRequest() { }
@@ -26,7 +26,7 @@ struct FPortalApplicationWindowNavigateToResponse
 {
 	GENERATED_USTRUCT_BODY()
 
-	UPROPERTY()
+	UPROPERTY(EditAnywhere, Category="Message")
 	bool Result;
 
 	FPortalApplicationWindowNavigateToResponse() { }
@@ -35,4 +35,5 @@ struct FPortalApplicationWindowNavigateToResponse
 	{ }
 };
 
+
 DECLARE_RPC(FPortalApplicationWindowNavigateTo, bool)
diff --git a/Engine/Source/Runtime/Portal/Messages/Public/PortalPackageInstallerMessages.h b/Engine/Source/Runtime/Portal/Messages/Public/PortalPackageInstallerMessages.h
--- a/Engine/Source/Runtime/Portal/Messages/Public/PortalPackageInstallerMessages.h
+++ b/Engine/Source/Runtime/Portal/Messages/Public/PortalPackageInstallerMessages.h
@@ -15,13 +15,13 @@ struct FPortalPackageInstallerInstallRequest
 {
 	GENERATED_USTRUCT_BODY()
 
-	UPROPERTY()
+	UPROPERTY(EditAnywhere, Category="Message")
 	FString AppName;
 
-	UPROPERTY()
+	UPROPERTY(EditAnywhere, Category="Message")
 	FString BuildLabel;
 
-	UPROPERTY()
+	UPROPERTY(EditAnywhere, Category="Message")
 	FString DestinationPath;
 
 	FPortalPackageInstallerInstallRequest() { }
@@ -39,7 +39,7 @@ struct FPortalPackageInstallerInstallResponse
 {
 	GENERATED_USTRUCT_BODY()
 
-	UPROPERTY()
+	UPROPERTY(EditAnywhere, Category="Message")
 	bool Result;
 
 	FPortalPackageInstallerInstallResponse() { }
@@ -55,16 +55,16 @@ struct FPortalPackageInstallerUninstallRequest
 {
 	GENERATED_USTRUCT_BODY()
 
-	UPROPERTY()
+	UPROPERTY(EditAnywhere, Category="Message")
 	FString AppName;
 
-	UPROPERTY()
+	UPROPERTY(EditAnywhere, Category="Message")
 	FString BuildLabel;
 
-	UPROPERTY()
+	UPROPERTY(EditAnywhere, Category="Message")
 	FString InstallationPath;
 
-	UPROPERTY()
+	UPROPERTY(EditAnywhere, Category="Message")
 	bool RemoveUserFiles;
 
 	FPortalPackageInstallerUninstallRequest() { }
@@ -83,7 +83,7 @@ struct FPortalPackageInstallerUninstallResponse
 {
 	GENERATED_USTRUCT_BODY()
 
-	UPROPERTY()
+	UPROPERTY(EditAnywhere, Category="Message")
 	bool Result;
 
 	FPortalPackageInstallerUninstallResponse() { }
diff --git a/Engine/Source/Runtime/Portal/Rpc/Private/PortalRpcMessages.h b/Engine/Source/Runtime/Portal/Rpc/Private/PortalRpcMessages.h
--- a/Engine/Source/Runtime/Portal/Rpc/Private/PortalRpcMessages.h
+++ b/Engine/Source/Runtime/Portal/Rpc/Private/PortalRpcMessages.h
@@ -11,11 +11,11 @@ struct FPortalRpcLocateServer
 	GENERATED_USTRUCT_BODY()
 
 	/** The product's unique identifier. */
-	UPROPERTY()
+	UPROPERTY(EditAnywhere, Category="Message")
 	FGuid ProductId;
 
 	/** The product's version string. */
-	UPROPERTY()
+	UPROPERTY(EditAnywhere, Category="Message")
 	FString ProductVersion;
 
 	/** Default constructor. */
@@ -35,7 +35,7 @@ struct FPortalRpcServer
 	GENERATED_USTRUCT_BODY()
 
 	/** The RPC server's message address as a string. */
-	UPROPERTY()
+	UPROPERTY(EditAnywhere, Category="Message")
 	FString ServerAddress;
 
 	/** Default constructor. */
