SessionFrontend: Using full username to identify session owners

diff --git a/Engine/Source/Developer/AutomationWindow/Private/SAutomationWindow.cpp b/Engine/Source/Developer/AutomationWindow/Private/SAutomationWindow.cpp
--- a/Engine/Source/Developer/AutomationWindow/Private/SAutomationWindow.cpp
+++ b/Engine/Source/Developer/AutomationWindow/Private/SAutomationWindow.cpp
@@ -1556,7 +1556,7 @@ void SAutomationWindow::FindWorkers()
 {
 	ActiveSession = SessionManager->GetSelectedSession();
 
-	bool SessionIsValid = ActiveSession.IsValid() && (ActiveSession->GetSessionOwner() == FPlatformProcess::UserName(true));
+	bool SessionIsValid = ActiveSession.IsValid() && (ActiveSession->GetSessionOwner() == FPlatformProcess::UserName(false));
 
 	if (SessionIsValid)
 	{
diff --git a/Engine/Source/Developer/DeviceManager/Private/Widgets/Browser/SDeviceBrowserDeviceListRow.h b/Engine/Source/Developer/DeviceManager/Private/Widgets/Browser/SDeviceBrowserDeviceListRow.h
--- a/Engine/Source/Developer/DeviceManager/Private/Widgets/Browser/SDeviceBrowserDeviceListRow.h
+++ b/Engine/Source/Developer/DeviceManager/Private/Widgets/Browser/SDeviceBrowserDeviceListRow.h
@@ -131,7 +131,7 @@ private:
 	{
 		FString ClaimUser = DeviceService->GetClaimUser();
 
-		if ((ClaimUser == FPlatformProcess::UserName(true)) && !DeviceService->IsRunning())
+		if ((ClaimUser == FPlatformProcess::UserName(false)) && !DeviceService->IsRunning())
 		{
 			ClaimUser += LOCTEXT("RemotelyHint", " (remotely)").ToString();
 		}
diff --git a/Engine/Source/Developer/LauncherServices/Private/Launcher/LauncherWorker.cpp b/Engine/Source/Developer/LauncherServices/Private/Launcher/LauncherWorker.cpp
--- a/Engine/Source/Developer/LauncherServices/Private/Launcher/LauncherWorker.cpp
+++ b/Engine/Source/Developer/LauncherServices/Private/Launcher/LauncherWorker.cpp
@@ -316,9 +316,9 @@ FString FLauncherWorker::CreateUATCommand( const ILauncherProfileRef& InProfile,
 		*InitialMap);
 
 	// additional commands to be sent to the commandline
-	FString AdditionalCommandLine = FString::Printf(TEXT(" -addcmdline=\"-SessionId=%s -SessionOwner=%s -SessionName='%s'%s\""),
+	FString AdditionalCommandLine = FString::Printf(TEXT(" -addcmdline=\"-SessionId=%s -SessionOwner='%s' -SessionName='%s'%s\""),
 		*SessionId.ToString(),
-		FPlatformProcess::UserName(true),
+		FPlatformProcess::UserName(false),
 		*InProfile->GetName(),
 		*RoleCommands);
 
diff --git a/Engine/Source/Developer/Profiler/Private/ProfilerManager.cpp b/Engine/Source/Developer/Profiler/Private/ProfilerManager.cpp
--- a/Engine/Source/Developer/Profiler/Private/ProfilerManager.cpp
+++ b/Engine/Source/Developer/Profiler/Private/ProfilerManager.cpp
@@ -405,7 +405,7 @@ void FProfilerManager::SessionManager_OnInstanceSelectionChanged(const TSharedPt
 {
 	const ISessionInfoPtr& SelectedSession = SessionManager->GetSelectedSession();
 	const bool SessionIsValid = SelectedSession.IsValid()
-		&& (SelectedSession->GetSessionOwner() == FPlatformProcess::UserName(true))
+		&& (SelectedSession->GetSessionOwner() == FPlatformProcess::UserName(false))
 		&& (SessionManager->GetSelectedInstances().Num() > 0);
 
 	if( ActiveSession != SelectedSession || FProfilerManager::GetSettings().bSingleInstanceMode )
diff --git a/Engine/Source/Developer/SessionFrontend/Private/Widgets/Browser/SSessionBrowser.cpp b/Engine/Source/Developer/SessionFrontend/Private/Widgets/Browser/SSessionBrowser.cpp
--- a/Engine/Source/Developer/SessionFrontend/Private/Widgets/Browser/SSessionBrowser.cpp
+++ b/Engine/Source/Developer/SessionFrontend/Private/Widgets/Browser/SSessionBrowser.cpp
@@ -170,7 +170,7 @@ void SSessionBrowser::FilterSessions()
 	for (auto& SessionInfo : AvailableSessions)
 	{
 		// process session
-		bool LocalOwner = SessionInfo->GetSessionOwner() == FPlatformProcess::UserName(true);
+		bool LocalOwner = SessionInfo->GetSessionOwner() == FPlatformProcess::UserName(false);
 
 		if (!SessionInfo->IsStandalone() && !LocalOwner)
 		{
@@ -451,7 +451,7 @@ FReply SSessionBrowser::HandleTerminateSessionButtonClicked()
 
 		if (SelectedSession.IsValid())
 		{
-			if (SelectedSession->GetSessionOwner() == FPlatformProcess::UserName(true))
+			if (SelectedSession->GetSessionOwner() == FPlatformProcess::UserName(false))
 			{
 				SelectedSession->Terminate();
 			}
diff --git a/Engine/Source/Developer/SessionFrontend/Private/Widgets/Browser/SSessionBrowserTreeSessionRow.h b/Engine/Source/Developer/SessionFrontend/Private/Widgets/Browser/SSessionBrowserTreeSessionRow.h
--- a/Engine/Source/Developer/SessionFrontend/Private/Widgets/Browser/SSessionBrowserTreeSessionRow.h
+++ b/Engine/Source/Developer/SessionFrontend/Private/Widgets/Browser/SSessionBrowserTreeSessionRow.h
@@ -81,7 +81,7 @@ protected:
 		// return name of a launched session
 		if (!SessionName.IsEmpty())
 		{
-			if (SessionInfo->GetSessionOwner() != FPlatformProcess::UserName(true))
+			if (SessionInfo->GetSessionOwner() != FPlatformProcess::UserName(false))
 			{
 				return FText::Format(LOCTEXT("SessionNameFormat", "{0} - {1}"), FText::FromString(SessionName), FText::FromString(SessionOwner));
 			}
@@ -103,7 +103,7 @@ protected:
 			}
 		}
 
-		if (SessionInfo->GetSessionOwner() != FPlatformProcess::UserName(true))
+		if (SessionInfo->GetSessionOwner() != FPlatformProcess::UserName(false))
 		{
 			return FText::Format(LOCTEXT("UnnamedSessionFormat", "Unnamed - {0}"), FText::FromString(SessionOwner));
 		}
diff --git a/Engine/Source/Developer/SessionServices/Private/SessionInstanceInfo.cpp b/Engine/Source/Developer/SessionServices/Private/SessionInstanceInfo.cpp
--- a/Engine/Source/Developer/SessionServices/Private/SessionInstanceInfo.cpp
+++ b/Engine/Source/Developer/SessionServices/Private/SessionInstanceInfo.cpp
@@ -65,7 +65,7 @@ void FSessionInstanceInfo::ExecuteCommand( const FString& CommandString )
 {
 	if (MessageEndpoint.IsValid() && EngineAddress.IsValid())
 	{
-		MessageEndpoint->Send(new FEngineServiceExecuteCommand(CommandString, FPlatformProcess::UserName(true)), EngineAddress);
+		MessageEndpoint->Send(new FEngineServiceExecuteCommand(CommandString, FPlatformProcess::UserName(false)), EngineAddress);
 	}
 }
 
@@ -74,7 +74,7 @@ void FSessionInstanceInfo::Terminate()
 {
 	if (MessageEndpoint.IsValid() && EngineAddress.IsValid())
 	{
-		MessageEndpoint->Send(new FEngineServiceTerminate(FPlatformProcess::UserName(true)), EngineAddress);
+		MessageEndpoint->Send(new FEngineServiceTerminate(FPlatformProcess::UserName(false)), EngineAddress);
 	}
 }
 
diff --git a/Engine/Source/Developer/SessionServices/Private/SessionManager.cpp b/Engine/Source/Developer/SessionServices/Private/SessionManager.cpp
--- a/Engine/Source/Developer/SessionServices/Private/SessionManager.cpp
+++ b/Engine/Source/Developer/SessionServices/Private/SessionManager.cpp
@@ -187,7 +187,7 @@ void FSessionManager::FindExpiredSessions(const FDateTime& Now)
 
 bool FSessionManager::IsValidOwner(const FString& Owner)
 {
-	if (Owner == FPlatformProcess::UserName(true))
+	if (Owner == FPlatformProcess::UserName(false))
 	{
 		return true;
 	}
diff --git a/Engine/Source/Developer/TargetDeviceServices/Private/TargetDeviceProxy.cpp b/Engine/Source/Developer/TargetDeviceServices/Private/TargetDeviceProxy.cpp
--- a/Engine/Source/Developer/TargetDeviceServices/Private/TargetDeviceProxy.cpp
+++ b/Engine/Source/Developer/TargetDeviceServices/Private/TargetDeviceProxy.cpp
@@ -234,19 +234,19 @@ bool FTargetDeviceProxy::LaunchApp(FName InVariant, const FString& AppId, EBuild
 
 void FTargetDeviceProxy::PowerOff(bool Force)
 {
-	MessageEndpoint->Send(new FTargetDeviceServicePowerOff(FPlatformProcess::UserName(true), Force), MessageAddress);
+	MessageEndpoint->Send(new FTargetDeviceServicePowerOff(FPlatformProcess::UserName(false), Force), MessageAddress);
 }
 
 
 void FTargetDeviceProxy::PowerOn()
 {
-	MessageEndpoint->Send(new FTargetDeviceServicePowerOn(FPlatformProcess::UserName(true)), MessageAddress);
+	MessageEndpoint->Send(new FTargetDeviceServicePowerOn(FPlatformProcess::UserName(false)), MessageAddress);
 }
 
 
 void FTargetDeviceProxy::Reboot()
 {
-	MessageEndpoint->Send(new FTargetDeviceServiceReboot(FPlatformProcess::UserName(true)), MessageAddress);
+	MessageEndpoint->Send(new FTargetDeviceServiceReboot(FPlatformProcess::UserName(false)), MessageAddress);
 }
 
 
diff --git a/Engine/Source/Developer/TargetDeviceServices/Private/TargetDeviceProxyManager.cpp b/Engine/Source/Developer/TargetDeviceServices/Private/TargetDeviceProxyManager.cpp
--- a/Engine/Source/Developer/TargetDeviceServices/Private/TargetDeviceProxyManager.cpp
+++ b/Engine/Source/Developer/TargetDeviceServices/Private/TargetDeviceProxyManager.cpp
@@ -75,7 +75,7 @@ void FTargetDeviceProxyManager::GetProxies(FName TargetPlatformName, bool Includ
 	{
 		const TSharedPtr<FTargetDeviceProxy>& Proxy = It.Value();
 
-		if ((IncludeUnshared || Proxy->IsShared()) || (Proxy->GetHostUser() == FPlatformProcess::UserName(true)))
+		if ((IncludeUnshared || Proxy->IsShared()) || (Proxy->GetHostUser() == FPlatformProcess::UserName(false)))
 		{
 			if (TargetPlatformName == NAME_None || Proxy->HasTargetPlatform(TargetPlatformName))
 			{
@@ -108,7 +108,7 @@ void FTargetDeviceProxyManager::SendPing()
 {
 	if (MessageEndpoint.IsValid())
 	{
-		MessageEndpoint->Publish(new FTargetDeviceServicePing(FPlatformProcess::UserName(true)), EMessageScope::Network);
+		MessageEndpoint->Publish(new FTargetDeviceServicePing(FPlatformProcess::UserName(false)), EMessageScope::Network);
 	}
 }
 
diff --git a/Engine/Source/Developer/TargetDeviceServices/Private/TargetDeviceService.cpp b/Engine/Source/Developer/TargetDeviceServices/Private/TargetDeviceService.cpp
--- a/Engine/Source/Developer/TargetDeviceServices/Private/TargetDeviceService.cpp
+++ b/Engine/Source/Developer/TargetDeviceServices/Private/TargetDeviceService.cpp
@@ -175,7 +175,7 @@ bool FTargetDeviceService::Start()
 		// notify other services
 		ClaimAddress = MessageEndpoint->GetAddress();
 		ClaimHost = FPlatformProcess::ComputerName();
-		ClaimUser = FPlatformProcess::UserName(true);
+		ClaimUser = FPlatformProcess::UserName(false);
 
 		MessageEndpoint->Publish(new FTargetDeviceClaimed(DeviceName, ClaimHost, ClaimUser));
 
@@ -191,7 +191,7 @@ void FTargetDeviceService::Stop()
 	if (Running)
 	{
 		// notify other services
-		MessageEndpoint->Publish(new FTargetDeviceUnclaimed(DeviceName, FPlatformProcess::ComputerName(), FPlatformProcess::UserName(true)));
+		MessageEndpoint->Publish(new FTargetDeviceUnclaimed(DeviceName, FPlatformProcess::ComputerName(), FPlatformProcess::UserName(false)));
 
 		Running = false;
 	}
@@ -282,7 +282,7 @@ void FTargetDeviceService::HandleClaimedMessage(const FTargetDeviceClaimed& Mess
 	{
 		if (Context->GetSender() != MessageEndpoint->GetAddress())
 		{
-			MessageEndpoint->Send(new FTargetDeviceClaimDenied(DeviceName, FPlatformProcess::ComputerName(), FPlatformProcess::UserName(true)), Context->GetSender());
+			MessageEndpoint->Send(new FTargetDeviceClaimDenied(DeviceName, FPlatformProcess::ComputerName(), FPlatformProcess::UserName(false)), Context->GetSender());
 		}
 	}
 	else
@@ -385,7 +385,7 @@ void FTargetDeviceService::HandlePingMessage(const FTargetDeviceServicePing& InM
 		return;
 	}
 
-	if (!Shared && (InMessage.HostUser != FPlatformProcess::UserName(true)))
+	if (!Shared && (InMessage.HostUser != FPlatformProcess::UserName(false)))
 	{
 		return;
 	}
@@ -402,7 +402,7 @@ void FTargetDeviceService::HandlePingMessage(const FTargetDeviceServicePing& InM
 		Message->Name = DefaultDevice->GetName();
 		Message->Type = TargetDeviceTypes::ToString(DefaultDevice->GetDeviceType());
 		Message->HostName = FPlatformProcess::ComputerName();
-		Message->HostUser = FPlatformProcess::UserName(true);
+		Message->HostUser = FPlatformProcess::UserName(false);
 		Message->Connected = DefaultDevice->IsConnected();
 		Message->Make = TEXT("@todo");
 		Message->Model = TEXT("@todo");
diff --git a/Engine/Source/Editor/UnrealEd/Private/Kismet2/DebuggerCommands.cpp b/Engine/Source/Editor/UnrealEd/Private/Kismet2/DebuggerCommands.cpp
--- a/Engine/Source/Editor/UnrealEd/Private/Kismet2/DebuggerCommands.cpp
+++ b/Engine/Source/Editor/UnrealEd/Private/Kismet2/DebuggerCommands.cpp
@@ -680,7 +680,7 @@ TSharedRef< SWidget > FPlayWorldCommands::GenerateLaunchMenuContent( TSharedRef<
 						{
 							LabelArguments.Add(TEXT("HostUser"), LOCTEXT("DisconnectedHint", " [Disconnected]"));
 						}
-						else if (DeviceProxy->GetHostUser() != FPlatformProcess::UserName(true))
+						else if (DeviceProxy->GetHostUser() != FPlatformProcess::UserName(false))
 						{
 							LabelArguments.Add(TEXT("HostUser"), FText::FromString(DeviceProxy->GetHostUser()));
 						}
diff --git a/Engine/Source/Runtime/Core/Private/Misc/App.cpp b/Engine/Source/Runtime/Core/Private/Misc/App.cpp
--- a/Engine/Source/Runtime/Core/Private/Misc/App.cpp
+++ b/Engine/Source/Runtime/Core/Private/Misc/App.cpp
@@ -113,7 +113,7 @@ void FApp::InitializeSession()
 
 	if (!FParse::Value(FCommandLine::Get(), TEXT("-SessionOwner="), SessionOwner))
 	{
-		SessionOwner = FPlatformProcess::UserName(true);
+		SessionOwner = FPlatformProcess::UserName(false);
 	}
 }
 
