Copying //UE4/Dev-Sequencer to //UE4/Main (Source: //UE4/Dev-Sequencer @ 3617642)

diff --git a/Engine/Source/Developer/TargetDeviceServices/Private/Services/TargetDeviceService.cpp b/Engine/Source/Developer/TargetDeviceServices/Private/Services/TargetDeviceService.cpp
new file mode 100644
--- /dev/null
+++ b/Engine/Source/Developer/TargetDeviceServices/Private/Services/TargetDeviceService.cpp
@@ -0,0 +1,575 @@
+// Copyright 1998-2017 Epic Games, Inc. All Rights Reserved.
+
+#include "TargetDeviceService.h"
+
+#include "HAL/PlatformProcess.h"
+#include "HAL/FileManager.h"
+#include "IMessageBus.h"
+#include "Interfaces/ITargetDevice.h"
+#include "Interfaces/ITargetPlatform.h"
+#include "MessageEndpoint.h"
+#include "MessageEndpointBuilder.h"
+#include "Misc/Paths.h"
+#include "Misc/ConfigCacheIni.h"
+#include "PlatformInfo.h"
+#include "Serialization/Archive.h"
+
+#include "TargetDeviceServiceMessages.h"
+
+
+/* Local helpers
+ *****************************************************************************/
+
+struct FVariantSortCallback
+{
+	FORCEINLINE bool operator()(const ITargetDeviceWeakPtr& A, const ITargetDeviceWeakPtr& B) const
+	{
+		ITargetDevicePtr APtr = A.Pin();
+		ITargetDevicePtr BPtr = B.Pin();
+
+		return APtr->GetTargetPlatform().GetVariantPriority() > BPtr->GetTargetPlatform().GetVariantPriority();
+	}
+};
+
+
+/* FTargetDeviceService structors
+ *****************************************************************************/
+
+FTargetDeviceService::FTargetDeviceService(const FString& InDeviceName, const TSharedRef<IMessageBus, ESPMode::ThreadSafe>& InMessageBus)
+	: DeviceName(InDeviceName)
+	, Running(false)
+	, Shared(false)
+{
+	// initialize messaging
+	MessageEndpoint = FMessageEndpoint::Builder(FName(*FString::Printf(TEXT("FTargetDeviceService (%s)"), *DeviceName)), InMessageBus)
+		.Handling<FTargetDeviceClaimDenied>(this, &FTargetDeviceService::HandleClaimDeniedMessage)
+		.Handling<FTargetDeviceClaimed>(this, &FTargetDeviceService::HandleClaimedMessage)
+		.Handling<FTargetDeviceServiceDeployCommit>(this, &FTargetDeviceService::HandleDeployCommitMessage)
+		.Handling<FTargetDeviceServiceDeployFile>(this, &FTargetDeviceService::HandleDeployFileMessage)
+		.Handling<FTargetDeviceServiceLaunchApp>(this, &FTargetDeviceService::HandleLaunchAppMessage)
+		.Handling<FTargetDeviceServicePing>(this, &FTargetDeviceService::HandlePingMessage)
+		.Handling<FTargetDeviceServicePowerOff>(this, &FTargetDeviceService::HandlePowerOffMessage)
+		.Handling<FTargetDeviceServicePowerOn>(this, &FTargetDeviceService::HandlePowerOnMessage)
+		.Handling<FTargetDeviceServiceReboot>(this, &FTargetDeviceService::HandleRebootMessage)
+		.Handling<FTargetDeviceServiceRunExecutable>(this, &FTargetDeviceService::HandleRunExecutableMessage)
+		.Handling<FTargetDeviceUnclaimed>(this, &FTargetDeviceService::HandleUnclaimedMessage);
+
+	if (MessageEndpoint.IsValid())
+	{
+		MessageEndpoint->Subscribe<FTargetDeviceClaimed>();
+		MessageEndpoint->Subscribe<FTargetDeviceUnclaimed>();
+		MessageEndpoint->Subscribe<FTargetDeviceServicePing>();
+	}
+}
+
+
+FTargetDeviceService::~FTargetDeviceService()
+{
+	Stop();
+	FMessageEndpoint::SafeRelease(MessageEndpoint);
+}
+
+
+/* ITargetDeviceService interface
+ *****************************************************************************/
+
+void FTargetDeviceService::AddTargetDevice(TSharedPtr<ITargetDevice, ESPMode::ThreadSafe> InDevice)
+{
+	if (!InDevice.IsValid())
+	{
+		return;
+	}
+
+	FName Variant = FName(InDevice->GetTargetPlatform().PlatformName().GetCharArray().GetData());
+
+	if (DevicePlatformName == NAME_None)
+	{
+		// If this seems nasty your right!
+		// This is just one more nastiness in this class due to the fact that we intend to refactor the target platform stuff as a separate task.
+		const PlatformInfo::FPlatformInfo& Info = InDevice->GetTargetPlatform().GetPlatformInfo();
+		DevicePlatformName = Info.PlatformInfoName;
+		const PlatformInfo::FPlatformInfo* VanillaInfo = PlatformInfo::FindVanillaPlatformInfo(Info.VanillaPlatformName);
+		DevicePlatformDisplayName = VanillaInfo->DisplayName.ToString();
+		
+		// Sigh the hacks... Should be able to remove if platform info gets cleaned up.... Windows doesn't have a reasonable vanilla platform.
+		const FString VariableSplit(TEXT("("));
+		FString Full = VanillaInfo->DisplayName.ToString();
+		FString Left;
+		FString Right;
+		bool bSplit = Full.Split(VariableSplit, &Left, &Right);
+		DevicePlatformDisplayName = bSplit ? Left.Trim() : Full;
+	}
+
+	// double add, which due to the async nature of some device discovery can't be easily avoided.
+	if (!(TargetDevicePtrs.FindRef(Variant).IsValid()))
+	{
+		TargetDevicePtrs.Add(Variant, InDevice);
+
+		// sort and choose cache the default
+		TargetDevicePtrs.ValueSort(FVariantSortCallback());
+		auto DeviceIterator = TargetDevicePtrs.CreateIterator();
+
+		if (DeviceIterator)
+		{
+			DefaultDevicePtr = (*DeviceIterator).Value;
+		}
+		else
+		{
+			DefaultDevicePtr = nullptr;
+		}
+	}
+}
+
+
+bool FTargetDeviceService::CanStart(FName InFlavor) const
+{
+	ITargetDevicePtr TargetDevice = GetDevice(InFlavor);
+
+	if (TargetDevice.IsValid())
+	{
+		return TargetDevice->IsConnected();
+	}
+
+	return false;
+}
+
+
+const FString& FTargetDeviceService::GetClaimHost()
+{
+	return ClaimHost;
+}
+
+
+const FString& FTargetDeviceService::GetClaimUser()
+{
+	return ClaimUser;
+}
+
+
+ITargetDevicePtr FTargetDeviceService::GetDevice(FName InVariant) const
+{
+	ITargetDevicePtr TargetDevice;
+
+	if (InVariant == NAME_None)
+	{
+		TargetDevice = DefaultDevicePtr.Pin();
+	}
+	else
+	{
+		const ITargetDeviceWeakPtr * WeakTargetDevicePtr = TargetDevicePtrs.Find(InVariant);
+
+		if (WeakTargetDevicePtr != nullptr)
+		{
+			TargetDevice = WeakTargetDevicePtr->Pin();
+		}
+	}
+
+	return TargetDevice;
+}
+
+
+FString FTargetDeviceService::GetDeviceName() const
+{
+	return DeviceName;
+}
+
+
+FName FTargetDeviceService::GetDevicePlatformName() const
+{
+	return DevicePlatformName;
+}
+
+
+FString FTargetDeviceService::GetDevicePlatformDisplayName() const
+{
+	return DevicePlatformDisplayName;
+}
+
+
+bool FTargetDeviceService::IsRunning() const
+{
+	return Running;
+}
+
+
+bool FTargetDeviceService::IsShared() const
+{
+	return (Running && Shared);
+}
+
+
+int32 FTargetDeviceService::NumTargetDevices()
+{
+	return TargetDevicePtrs.Num();
+}
+
+
+void FTargetDeviceService::RemoveTargetDevice(TSharedPtr<ITargetDevice, ESPMode::ThreadSafe> InDevice)
+{
+	if (!InDevice.IsValid())
+	{
+		return;
+	}
+
+	FName Variant = FName(InDevice->GetTargetPlatform().PlatformName().GetCharArray().GetData());
+
+	TargetDevicePtrs.Remove(Variant);
+
+	// Cache the default
+	auto DeviceIterator = TargetDevicePtrs.CreateIterator();
+
+	if (DeviceIterator)
+	{
+		DefaultDevicePtr = (*DeviceIterator).Value;
+	}
+	else
+	{
+		DefaultDevicePtr = nullptr;
+	}
+}
+
+
+void FTargetDeviceService::SetShared(bool InShared)
+{
+	Shared = InShared;
+}
+
+
+bool FTargetDeviceService::Start()
+{
+	if (!Running && MessageEndpoint.IsValid())
+	{
+		// notify other services
+		ClaimAddress = MessageEndpoint->GetAddress();
+		ClaimHost = FPlatformProcess::ComputerName();
+		ClaimUser = FPlatformProcess::UserName(false);
+
+		MessageEndpoint->Publish(new FTargetDeviceClaimed(DeviceName, ClaimHost, ClaimUser));
+
+		Running = true;
+	}
+		
+	return true;
+}
+
+
+void FTargetDeviceService::Stop()
+{
+	if (Running)
+	{
+	
+		MessageEndpoint->Publish(new FTargetDeviceUnclaimed(DeviceName, FPlatformProcess::ComputerName(), FPlatformProcess::UserName(false)));
+        FPlatformProcess::SleepNoStats(0.01);
+
+        // Only stop the device if we care about device claiming
+		GConfig->GetBool(TEXT("/Script/Engine.Engine"), TEXT("DisableDeviceClaiming"), Running, GEngineIni);
+	}
+}
+
+
+/* FTargetDeviceService implementation
+ *****************************************************************************/
+
+bool FTargetDeviceService::StoreDeployedFile(FArchive* FileReader, const FString& TargetFileName) const
+{
+	if (FileReader == nullptr)
+	{
+		return false;
+	}
+
+	// create target file
+	FArchive* FileWriter = IFileManager::Get().CreateFileWriter(*TargetFileName);
+
+	if (FileWriter == nullptr)
+	{
+		return false;
+	}
+
+	FileReader->Seek(0);
+
+	// copy file contents
+	int64 BytesRemaining = FileReader->TotalSize();
+	int32 BufferSize = 128 * 1024;
+
+	if (BytesRemaining < BufferSize)
+	{
+		BufferSize = BytesRemaining;
+	}
+
+	void* Buffer = FMemory::Malloc(BufferSize);
+
+	while (BytesRemaining > 0)
+	{
+		FileReader->Serialize(Buffer, BufferSize);
+		FileWriter->Serialize(Buffer, BufferSize);
+
+		BytesRemaining -= BufferSize;
+
+		if (BytesRemaining < BufferSize)
+		{
+			BufferSize = BytesRemaining;
+		}
+	}
+
+	// clean up
+	FMemory::Free(Buffer);
+	delete FileWriter;
+
+	return true;
+}
+
+
+/* FTargetDeviceService callbacks
+ *****************************************************************************/
+
+void FTargetDeviceService::HandleClaimDeniedMessage(const FTargetDeviceClaimDenied& Message, const TSharedRef<IMessageContext, ESPMode::ThreadSafe>& Context)
+{
+// HACK: Disabling claim denied message. Allows the editor to always claim a device and should prevent cases where instances of the
+// editor running on other machines can steal a device from us - which is undesirable on some platforms.
+// Also see FTargetDeviceProxyManager::HandlePongMessage()
+#if 0
+ 	if (Running && (Message.DeviceName == DeviceName))
+ 	{
+ 		Stop();
+ 
+ 		ClaimAddress = Context->GetSender();
+ 		ClaimHost = Message.HostName;
+ 		ClaimUser = Message.HostUser;
+ 	}
+#endif
+}
+
+
+void FTargetDeviceService::HandleClaimedMessage(const FTargetDeviceClaimed& Message, const TSharedRef<IMessageContext, ESPMode::ThreadSafe>& Context)
+{
+	if (Message.DeviceName != DeviceName)
+	{
+		return;
+	}
+
+	if (Running)
+	{
+		if (Context->GetSender() != MessageEndpoint->GetAddress())
+		{
+			MessageEndpoint->Send(new FTargetDeviceClaimDenied(DeviceName, FPlatformProcess::ComputerName(), FPlatformProcess::UserName(false)), Context->GetSender());
+		}
+	}
+	else
+	{
+		ClaimAddress = Context->GetSender();
+		ClaimHost = Message.HostName;
+		ClaimUser = Message.HostUser;
+	}
+}
+
+
+void FTargetDeviceService::HandleUnclaimedMessage(const FTargetDeviceUnclaimed& Message, const TSharedRef<IMessageContext, ESPMode::ThreadSafe>& Context)
+{
+	if (Message.DeviceName == DeviceName)
+	{
+		if (Context->GetSender() == ClaimAddress)
+		{
+			ClaimAddress.Invalidate();
+			ClaimHost.Empty();
+			ClaimUser.Empty();
+		}
+	}
+}
+
+
+void FTargetDeviceService::HandleDeployFileMessage(const FTargetDeviceServiceDeployFile& Message, const TSharedRef<IMessageContext, ESPMode::ThreadSafe>& Context)
+{
+	if (!Running)
+	{
+		return;
+	}
+
+	TSharedPtr<IMessageAttachment, ESPMode::ThreadSafe> Attachment = Context->GetAttachment();
+
+	if (Attachment.IsValid())
+	{
+		FArchive* FileReader = Attachment->CreateReader();
+
+		if (FileReader != nullptr)
+		{
+			FString DeploymentFolder = FPaths::EngineIntermediateDir() / TEXT("Deploy") / Message.TransactionId.ToString();
+			FString TargetPath = DeploymentFolder / Message.TargetFileName;
+
+			StoreDeployedFile(FileReader, TargetPath);
+
+			delete FileReader;
+		}
+	}
+}
+
+
+void FTargetDeviceService::HandleDeployCommitMessage(const FTargetDeviceServiceDeployCommit& Message, const TSharedRef<IMessageContext, ESPMode::ThreadSafe>& Context)
+{
+	if (!Running)
+	{
+		return;
+	}
+
+	ITargetDevicePtr TargetDevice = GetDevice(Message.Variant);
+
+	if (TargetDevice.IsValid())
+	{
+		FString SourceFolder = FPaths::EngineIntermediateDir() / TEXT("Deploy") / Message.TransactionId.ToString();
+		FString OutAppId;
+
+		bool Succeeded = TargetDevice->Deploy(SourceFolder, OutAppId);
+
+		IFileManager::Get().DeleteDirectory(*SourceFolder, false, true);
+		MessageEndpoint->Send(new FTargetDeviceServiceDeployFinished(Message.Variant, OutAppId, Succeeded, Message.TransactionId), Context->GetSender());
+	}
+}
+
+
+void FTargetDeviceService::HandleLaunchAppMessage(const FTargetDeviceServiceLaunchApp& Message, const TSharedRef<IMessageContext, ESPMode::ThreadSafe>& Context)
+{
+	if (!Running)
+	{
+		return;
+	}
+
+	ITargetDevicePtr TargetDevice = GetDevice(Message.Variant);
+
+	if (TargetDevice.IsValid())
+	{
+		uint32 ProcessId;
+		bool Succeeded = TargetDevice->Launch(Message.AppID, (EBuildConfigurations::Type)Message.BuildConfiguration, EBuildTargets::Game, Message.Params, &ProcessId);
+
+		if (MessageEndpoint.IsValid())
+		{
+			MessageEndpoint->Send(new FTargetDeviceServiceLaunchFinished(Message.AppID, ProcessId, Succeeded), Context->GetSender());
+		}
+	}
+}
+
+
+void FTargetDeviceService::HandlePingMessage(const FTargetDeviceServicePing& InMessage, const TSharedRef<IMessageContext, ESPMode::ThreadSafe>& Context)
+{
+	if (!Running)
+	{
+		return;
+	}
+
+	if (!Shared && (InMessage.HostUser != FPlatformProcess::UserName(false)))
+	{
+		return;
+	}
+
+	ITargetDevicePtr DefaultDevice = GetDevice(); // Default Device is needed here!
+
+	if (DefaultDevice.IsValid())
+	{
+		const FString& PlatformName = DefaultDevice->GetTargetPlatform().PlatformName();
+		const PlatformInfo::FPlatformInfo* VanillaInfo = PlatformInfo::FindVanillaPlatformInfo(FName(*PlatformName));
+
+		FTargetDeviceServicePong* Message = new FTargetDeviceServicePong();
+
+		Message->Name = DefaultDevice->GetName();
+		Message->Type = TargetDeviceTypes::ToString(DefaultDevice->GetDeviceType());
+		Message->HostName = FPlatformProcess::ComputerName();
+		Message->HostUser = FPlatformProcess::UserName(false);
+		Message->Connected = DefaultDevice->IsConnected();
+		Message->Make = TEXT("@todo");
+		Message->Model = TEXT("@todo");
+		DefaultDevice->GetUserCredentials(Message->DeviceUser, Message->DeviceUserPassword);
+		Message->Shared = Shared;
+		Message->SupportsMultiLaunch = DefaultDevice->SupportsFeature(ETargetDeviceFeatures::MultiLaunch);
+		Message->SupportsPowerOff = DefaultDevice->SupportsFeature(ETargetDeviceFeatures::PowerOff);
+		Message->SupportsPowerOn = DefaultDevice->SupportsFeature(ETargetDeviceFeatures::PowerOn);
+		Message->SupportsReboot = DefaultDevice->SupportsFeature(ETargetDeviceFeatures::Reboot);
+		Message->SupportsVariants = DefaultDevice->GetTargetPlatform().SupportsVariants();
+		Message->DefaultVariant = FName(DefaultDevice->GetTargetPlatform().PlatformName().GetCharArray().GetData());
+
+		// Add the data for all the flavors
+		Message->Variants.SetNumZeroed(TargetDevicePtrs.Num());
+
+		int Index = 0;
+		for (auto TargetDeviceIt = TargetDevicePtrs.CreateIterator(); TargetDeviceIt; ++TargetDeviceIt, ++Index)
+		{
+			const ITargetDevicePtr& TargetDevice = TargetDeviceIt.Value().Pin();
+			const PlatformInfo::FPlatformInfo& Info = TargetDevice->GetTargetPlatform().GetPlatformInfo();
+
+			FTargetDeviceVariant& Variant = Message->Variants[Index];
+
+			Variant.DeviceID = TargetDevice->GetId().ToString();
+			Variant.VariantName = TargetDeviceIt.Key();
+			Variant.TargetPlatformName = TargetDevice->GetTargetPlatform().PlatformName();
+			Variant.TargetPlatformId = Info.TargetPlatformName;
+			Variant.VanillaPlatformId = Info.VanillaPlatformName;
+			Variant.PlatformDisplayName = Info.DisplayName.ToString();
+		}
+
+		MessageEndpoint->Send(Message, Context->GetSender());
+	}
+}
+
+
+void FTargetDeviceService::HandlePowerOffMessage( const FTargetDeviceServicePowerOff& Message, const TSharedRef<IMessageContext, ESPMode::ThreadSafe>& Context )
+{
+	if (!Running)
+	{
+		return;
+	}
+
+	ITargetDevicePtr TargetDevice = GetDevice(); // Any Device is fine here
+
+	if (TargetDevice.IsValid())
+	{
+		TargetDevice->PowerOff(Message.Force);
+	}
+}
+
+
+void FTargetDeviceService::HandlePowerOnMessage( const FTargetDeviceServicePowerOn& Message, const TSharedRef<IMessageContext, ESPMode::ThreadSafe>& Context )
+{
+	if (!Running)
+	{
+		return;
+	}
+
+	ITargetDevicePtr TargetDevice = GetDevice(); // Any Device is fine here
+
+	if (TargetDevice.IsValid())
+	{
+		TargetDevice->PowerOn();
+	}
+}
+
+
+void FTargetDeviceService::HandleRebootMessage( const FTargetDeviceServiceReboot& Message, const TSharedRef<IMessageContext, ESPMode::ThreadSafe>& Context )
+{
+	if (!Running)
+	{
+		return;
+	}
+
+	ITargetDevicePtr TargetDevice = GetDevice(); // Any Device is fine here
+
+	if (TargetDevice.IsValid())
+	{
+		TargetDevice->Reboot();
+	}
+}
+
+
+void FTargetDeviceService::HandleRunExecutableMessage(const FTargetDeviceServiceRunExecutable& Message, const TSharedRef<IMessageContext, ESPMode::ThreadSafe>& Context)
+{
+	if (!Running)
+	{
+		return;
+	}
+
+	ITargetDevicePtr TargetDevice = GetDevice(Message.Variant);
+
+	if (TargetDevice.IsValid())
+	{
+		uint32 OutProcessId;
+		bool Succeeded = TargetDevice->Run(Message.ExecutablePath, Message.Params, &OutProcessId);
+
+		MessageEndpoint->Send(new FTargetDeviceServiceRunFinished(Message.Variant, Message.ExecutablePath, OutProcessId, Succeeded), Context->GetSender());
+	}
+}
diff --git a/Engine/Source/Developer/TargetDeviceServices/Private/TargetDeviceServicesModule.cpp b/Engine/Source/Developer/TargetDeviceServices/Private/TargetDeviceServicesModule.cpp
--- a/Engine/Source/Developer/TargetDeviceServices/Private/TargetDeviceServicesModule.cpp
+++ b/Engine/Source/Developer/TargetDeviceServices/Private/TargetDeviceServicesModule.cpp
@@ -1,9 +1,11 @@
 // Copyright 1998-2017 Epic Games, Inc. All Rights Reserved.
 
-#include "CoreMinimal.h"
+#include "TargetDeviceServicesPrivate.h"
+
+#include "CoreTypes.h"
 #include "Modules/ModuleManager.h"
-#include "Interfaces/ITargetDeviceServicesModule.h"
-#include "TargetDeviceServicesLog.h"
+
+#include "ITargetDeviceServicesModule.h"
 #include "TargetDeviceProxyManager.h"
 #include "TargetDeviceServiceManager.h"
 
@@ -19,9 +21,9 @@ class FTargetDeviceServicesModule
 {
 public:
 
-	// ITargetDeviceServicesModule interface
+	//~ ITargetDeviceServicesModule interface
 
-	virtual ITargetDeviceProxyManagerRef GetDeviceProxyManager() override
+	virtual TSharedRef<ITargetDeviceProxyManager> GetDeviceProxyManager() override
 	{
 		if (!DeviceProxyManagerSingleton.IsValid())
 		{
@@ -31,7 +33,7 @@ public:
 		return DeviceProxyManagerSingleton.ToSharedRef();
 	}
 
-	virtual ITargetDeviceServiceManagerRef GetDeviceServiceManager() override
+	virtual TSharedRef<ITargetDeviceServiceManager> GetDeviceServiceManager() override
 	{
 		if (!DeviceServiceManagerSingleton.IsValid())
 		{
@@ -44,10 +46,10 @@ public:
 private:
 
 	/** Holds the device proxy manager singleton. */
-	ITargetDeviceProxyManagerPtr DeviceProxyManagerSingleton;
+	TSharedPtr<FTargetDeviceProxyManager> DeviceProxyManagerSingleton;
 
 	/** Holds the device service manager singleton. */
-	ITargetDeviceServiceManagerPtr DeviceServiceManagerSingleton;
+	TSharedPtr<FTargetDeviceServiceManager> DeviceServiceManagerSingleton;
 };
 
 
diff --git a/Engine/Source/Developer/TargetDeviceServices/Public/ITargetDeviceProxy.h b/Engine/Source/Developer/TargetDeviceServices/Public/ITargetDeviceProxy.h
new file mode 100644
--- /dev/null
+++ b/Engine/Source/Developer/TargetDeviceServices/Public/ITargetDeviceProxy.h
@@ -0,0 +1,348 @@
+// Copyright 1998-2017 Epic Games, Inc. All Rights Reserved.
+
+#pragma once
+
+#include "CoreTypes.h"
+#include "Containers/Map.h"
+#include "Containers/UnrealString.h"
+#include "Delegates/Delegate.h"
+#include "Templates/SharedPointer.h"
+#include "UObject/NameTypes.h"
+
+class ITargetDeviceProxy;
+
+struct FGuid;
+
+
+/** Type definition for shared pointers to instances of ITargetDeviceProxy. */
+DEPRECATED(4.16, "ITargetDeviceProxyPtr is deprecated. Please use 'TSharedPtr<ITargetDeviceProxy>' instead!")
+typedef TSharedPtr<class ITargetDeviceProxy> ITargetDeviceProxyPtr;
+
+/** Type definition for shared references to instances of ITargetDeviceProxy. */
+DEPRECATED(4.16, "ITargetDeviceProxyRef is deprecated. Please use 'TSharedPtr<ITargetDeviceProxy>' instead!")
+typedef TSharedRef<class ITargetDeviceProxy> ITargetDeviceProxyRef;
+
+
+/**
+ * Delegate type for committed build deployments.
+ *
+ * The first parameter is the transaction identifier of the completed deployment.
+ * The second parameter is the identifier of the deployed application (used for launching).
+ */
+DECLARE_MULTICAST_DELEGATE_TwoParams(FOnTargetDeviceProxyDeployCommitted, const FGuid&, const FString&)
+
+/**
+ * Delegate type for failed build deployments.
+ *
+ * The first parameter is the transaction identifier of the failed deployment.
+ */
+DECLARE_MULTICAST_DELEGATE_OneParam(FOnTargetDeviceProxyDeployFailed, const FGuid&)
+
+/**
+ * Delegate type for failed build launches.
+ *
+ * The first parameter is the application identifier of the failed launch.
+ */
+DECLARE_MULTICAST_DELEGATE_OneParam(FOnTargetDeviceProxyLaunchFailed, const FString&)
+
+/**
+ * Delegate type for successful build launches.
+ *
+ * The first parameter is the application identifier of the successful launch.
+ * The second parameter is the process identifier of the launched application.
+ */
+DECLARE_MULTICAST_DELEGATE_TwoParams(FOnTargetDeviceProxyLaunchSucceeded, const FString&, uint32)
+
+
+/**
+ * Interface for target device proxies.
+ */
+class ITargetDeviceProxy
+{
+public:
+
+	/**
+	 * Checks whether the device can launch multiple games or applications simultaneously.
+	 *
+	 * @return true if launching multiple instances is supported, false otherwise.
+	 */
+	virtual const bool CanMultiLaunch() const = 0;
+
+	/**
+	 * Checks whether this device can be powered off remotely.
+	 *
+	 * @return true if the device can be powered off, false otherwise.
+	 */
+	virtual bool CanPowerOff() const = 0;
+
+	/**
+	 * Checks whether this device can be powered on remotely.
+	 *
+	 * @return true if the device can be powered on, false otherwise.
+	 */
+	virtual bool CanPowerOn() const = 0;
+
+	/**
+	 * Checks whether this device can be rebooted remotely.
+	 *
+	 * @return true if the device can be rebooted, false otherwise.
+	 */
+	virtual bool CanReboot() const = 0;
+
+	/**
+	 * Checks whether this device's target platform supports variants.
+	 *
+	 * @return true if the device's target platform supports variants, false otherwise.
+	 */
+	virtual bool CanSupportVariants() const = 0;
+
+	/**
+	 * Gets the number of variants this device supports
+	 *
+	 * @return The number of variants supported.
+	 */
+	virtual int32 GetNumVariants() const = 0;
+
+	/**
+	 * Gets the list of variants this device supports
+	 *
+	 * @param OutVariants This list will be populated with the unique variant names. 
+	 * @return The number of variants in the list.
+	 */
+	virtual int32 GetVariants(TArray<FName>& OutVariants) const = 0;
+
+	/**
+	 * Checks whether this device proxy contains a variant
+	 *
+	 * @param InVariant Variant to check for.
+	 * @return true is this variant is found.
+	 */
+	virtual bool HasVariant(FName InVariant) const = 0;
+
+	/**
+	 * Gets the variant name of the target device.
+	 *
+	 * @param InDeviceId target device id to get the variant for.
+	 * @return The target device variant name.
+	 */
+	virtual FName GetTargetDeviceVariant(const FString& InDeviceId) const = 0;
+
+	/**
+	 * Checks whether this device proxy contains a variant for the provided target device.
+	 *
+	 * @param InDeviceId DeviceId to check for a variant.
+	 * @return true if this device has a variant of the provided DeviceId.
+	 */
+	virtual bool HasDeviceId(const FString& InDeviceId) const = 0;
+
+	/**
+	 * Gets the identifier of the device.
+	 *
+	 * @param InVariant name of the variant to return the device identifier for, If InVariant == NAME_None you'll get back the default target device identifier.
+	 * @return The target device identifier.
+	 */
+	virtual const FString& GetTargetDeviceId(FName InVariant) const = 0;
+
+	/**
+	* Checks whether this device proxy contains a variant for the provided platform.
+	*
+	* @param InTargetPlatformName Target Platform Name for the target device to check for as a variant of this device.
+	* @return true if this device has a variant of the provided Target Platform .
+	*/
+	virtual bool HasTargetPlatform(FName InTargetPlatformId) const = 0;
+
+	/**
+	 * Gets the target platform of the device variant.
+	 *
+	 * @param InVariant Variant to get the platform name for.
+	 * @return Target Platform Name string.
+	 */
+	virtual FString GetTargetPlatformName(FName InVariant) const = 0;
+
+	/**
+	 * Gets the target platform of the device variant.
+	 *
+	 * @param InVariant Variant to get the platform id for.
+	 * @return Target Platform Id.
+	 */
+	virtual FName GetTargetPlatformId(FName InVariant) const = 0;
+
+	/**
+	 * Gets the Vanilla platform of the device variant.
+	 *
+	 * @param InVariant Variant to get the platform id for.
+	 * @return Vanilla Platform Id.
+	 */
+	virtual FName GetVanillaPlatformId(FName InVariant) const = 0;
+
+	/**
+	 * Gets the Vanilla platform of the device variant.
+	 *
+	 * @param InVariant Variant to get the vanilla platform name for.
+	 * @return Platform Display Name string.
+	 */
+	virtual FText GetPlatformDisplayName(FName InVariant) const = 0;
+
+	/**
+	 * Gets the name of the host machine that claimed the device.
+	 *
+	 * @return Host name string.
+	 */
+	virtual const FString& GetHostName() const = 0;
+	 
+	/**
+	 * Gets the name of the user that claimed the device.
+	 *
+	 * @return User name.
+	 */
+	virtual const FString& GetHostUser() const = 0;
+
+	/**
+	 * Gets the name of the (device's) user that is logged in on a device.
+	 *
+	 * @return Device user name.
+	 */
+	virtual const FString& GetDeviceUser() const = 0;
+
+	/**
+	 * Gets the password of the (device's) user that is logged in on a device.
+	 *
+	 * @return Device user password.
+	 */
+	virtual const FString& GetDeviceUserPassword() const = 0;
+
+	/**
+	 * Gets the device make (i.e. Apple or Sony).
+	 *
+	 * @return Device make.
+	 */
+	virtual const FString& GetMake() const = 0;
+
+	/**
+	 * Gets the device model (i.e. PS3 or XBox).
+	 *
+	 * @return Device model.
+	 */
+	virtual const FString& GetModel() const = 0;
+	 
+	/**
+	 * Gets the name of the device (i.e. network name or IP address).
+	 *
+	 * @return Device name string.
+	 */
+	virtual const FString& GetName() const = 0;
+
+	/**
+	 * Gets the device type (i.e. Console, PC or Mobile).
+	 *
+	 * @return Device type.
+	 */
+	virtual const FString& GetType() const = 0;
+
+	/**
+	 * Checks whether the device is currently connected.
+	 *
+	 * @return true if the device is connected, false otherwise.
+	 */
+	virtual bool IsConnected() const = 0;
+
+	/**
+	 * Checks whether this device is being shared with other users.
+	 *
+	 * @return true if the device is being shared, false otherwise.
+	 */
+	virtual bool IsShared() const = 0;
+
+public:
+
+	/**
+	 * Deploys the specified list of files to the device.
+	 *
+	 * If the deployment succeeds, the OnDeployCommitted() delegate will return an
+	 * application identifier (AppId) that can be used to launch to the deployed
+	 * application. If deploymend fails, the OnDeployFailed() delegate will be invoked.
+	 *
+	 * @param InVariant Variant to deploy files on.
+	 * @param Files The files to deploy.
+	 * @param TransactionId The transaction identifier for the deployment.
+	 * @return true if deployment has been started, false otherwise.
+	 * @see Launch, OnDeployCommitted, OnDeployFailed
+	 */
+	virtual bool DeployApp(FName InVariant, const TMap<FString, FString>& Files, const FGuid& TransactionId) = 0;
+
+	/**
+	 * Launches a previously deployed build on the target device.
+	 *
+	 * If the launching fails, the OnLaunchFailed() delegate will be invoked.
+	 *
+	 * @param InVariant Variant to launch the app with.
+	 * @param AppId The identifier of the application to launch.
+	 * @param BuildConfiguration The build configuration to launch with.
+	 * @param Params The command line parameters to launch with.
+	 * @return true on success, false otherwise.
+	 * @see Deploy, OnLaunchFailed
+	 */
+	virtual bool LaunchApp(FName InVariant, const FString& AppId, EBuildConfigurations::Type BuildConfiguration, const FString& Params) = 0;
+
+	/**
+	 * Powers off the device.
+	 *
+	 * @param Force Whether to force powering off.
+	 * @see PowerOn, Reboot
+	 */
+	virtual void PowerOff(bool Force) = 0;
+
+	/**
+	 * Powers on the device.
+	 *
+	 * @see PowerOff, Reboot
+	 */
+	virtual void PowerOn() = 0;
+
+	/**
+	 * Reboots the device.
+	 *
+	 * @see PowerOff, PowerOn
+	 */
+	virtual void Reboot() = 0;
+
+	/**
+	 * Attempts to run an executable on the device.
+	 *
+	 * @param InVariant Variant to run the executable with.
+	 * @param ExecutablePath The path to the executable to run.
+	 * @param Params The command line parameters.
+	 */
+	virtual void Run(FName InVariant, const FString& ExecutablePath, const FString& Params) = 0;
+
+public:
+
+	/**
+	 * Returns a delegate to be invoked when a build has been deployed to the target device.
+	 *
+	 * @return The delegate.
+	 * @see DeployApp
+	 */
+	virtual FOnTargetDeviceProxyDeployCommitted& OnDeployCommitted() = 0;
+
+	/**
+	 * Returns a delegate to be invoked when a build failed to deploy to the target device.
+	 *
+	 * @return The delegate.
+	 * @see DeployApp
+	 */
+	virtual FOnTargetDeviceProxyDeployFailed& OnDeployFailed() = 0;
+
+	/**
+	 * Returns a delegate to be invoked when a build failed to launch on the target device.
+	 *
+	 * @return The delegate.
+	 * @see LaunchApp
+	 */
+	virtual FOnTargetDeviceProxyLaunchFailed& OnLaunchFailed() = 0;
+
+public:
+
+	/** Virtual destructor. */
+	virtual ~ITargetDeviceProxy() { }
+};
diff --git a/Engine/Source/Developer/TargetDeviceServices/TargetDeviceServices.Build.cs b/Engine/Source/Developer/TargetDeviceServices/TargetDeviceServices.Build.cs
--- a/Engine/Source/Developer/TargetDeviceServices/TargetDeviceServices.Build.cs
+++ b/Engine/Source/Developer/TargetDeviceServices/TargetDeviceServices.Build.cs
@@ -7,26 +7,29 @@ namespace UnrealBuildTool.Rules
 		public TargetDeviceServices(ReadOnlyTargetRules Target) : base(Target)
 		{
 			PublicDependencyModuleNames.AddRange(
-				new string[]
-				{
+				new string[] {
 					"Core",
 					"TargetPlatform",
-                    "DesktopPlatform",
-				}
-			);
+					"DesktopPlatform",
+				});
 
 			PrivateDependencyModuleNames.AddRange(
-				new string[]
-				{
+				new string[] {
 					"CoreUObject",
-				}
-			);
+				});
 
 			PrivateIncludePathModuleNames.AddRange(
 				new string[] {
 					"Messaging",
-				}
-			);
+					"MessagingCommon",
+				});
+
+			PrivateIncludePaths.AddRange(
+				new string[] {
+					"Developer/TargetDeviceServices/Private",
+					"Developer/TargetDeviceServices/Private/Proxies",
+					"Developer/TargetDeviceServices/Private/Services",
+				});
 		}
 	}
 }
diff --git a/Engine/Source/Developer/TextureFormatASTC/Private/TextureFormatASTC.cpp b/Engine/Source/Developer/TextureFormatASTC/Private/TextureFormatASTC.cpp
--- a/Engine/Source/Developer/TextureFormatASTC/Private/TextureFormatASTC.cpp
+++ b/Engine/Source/Developer/TextureFormatASTC/Private/TextureFormatASTC.cpp
@@ -9,7 +9,8 @@
 #include "Misc/Guid.h"
 #include "Misc/ConfigCacheIni.h"
 #include "ImageCore.h"
-#include "Interfaces/IImageWrapperModule.h"
+#include "IImageWrapper.h"
+#include "IImageWrapperModule.h"
 #include "Modules/ModuleManager.h"
 #include "Interfaces/ITextureFormat.h"
 #include "Interfaces/ITextureFormatModule.h"
@@ -173,7 +174,7 @@ static bool CompressSliceToASTC(
 	}
 	
 	// Compress and retrieve the PNG data to write out to disk
-	IImageWrapperPtr ImageWrapper = ImageWrapperModule.CreateImageWrapper(EImageFormat::PNG);
+	TSharedPtr<IImageWrapper> ImageWrapper = ImageWrapperModule.CreateImageWrapper(EImageFormat::PNG);
 	ImageWrapper->SetRaw(SourceData, SizeX * SizeY * 4, SizeX, SizeY, ERGBFormat::RGBA, 8);
 	const TArray<uint8>& FileData = ImageWrapper->GetCompressed();
 	int32 FileDataSize = FileData.Num();
diff --git a/Engine/Source/Editor/AddContentDialog/Private/ViewModels/ContentSourceViewModel.cpp b/Engine/Source/Editor/AddContentDialog/Private/ViewModels/ContentSourceViewModel.cpp
--- a/Engine/Source/Editor/AddContentDialog/Private/ViewModels/ContentSourceViewModel.cpp
+++ b/Engine/Source/Editor/AddContentDialog/Private/ViewModels/ContentSourceViewModel.cpp
@@ -1,17 +1,19 @@
 // Copyright 1998-2017 Epic Games, Inc. All Rights Reserved.
 
 #include "ViewModels/ContentSourceViewModel.h"
+
+#include "IImageWrapper.h"
+#include "IImageWrapperModule.h"
 #include "Internationalization/Culture.h"
 #include "Modules/ModuleManager.h"
 
-#include "Interfaces/IImageWrapperModule.h"
 
 #define LOCTEXT_NAMESPACE "ContentSourceViewModel"
 
 const FString DefaultLanguageCode = "en";
-
 uint32 FContentSourceViewModel::ImageID = 0;
 
+
 FContentSourceViewModel::FContentSourceViewModel(TSharedPtr<IContentSource> ContentSourceIn)
 {
 	ContentSource = ContentSourceIn;
@@ -104,7 +106,7 @@ TSharedPtr<FSlateDynamicImageBrush> FContentSourceViewModel::CreateBrushFromRawD
 	bool bSucceeded = false;
 	TArray<uint8> DecodedImage;
 	IImageWrapperModule& ImageWrapperModule = FModuleManager::LoadModuleChecked<IImageWrapperModule>(FName("ImageWrapper"));
-	IImageWrapperPtr ImageWrapper = ImageWrapperModule.CreateImageWrapper(EImageFormat::PNG);
+	TSharedPtr<IImageWrapper> ImageWrapper = ImageWrapperModule.CreateImageWrapper(EImageFormat::PNG);
 	if (ImageWrapper.IsValid() && (RawData.Num() > 0) && ImageWrapper->SetCompressed(RawData.GetData(), RawData.Num()))
 	{
 		Width = ImageWrapper->GetWidth();
diff --git a/Engine/Source/Editor/DetailCustomizations/Private/TimespanStructCustomization.cpp b/Engine/Source/Editor/DetailCustomizations/Private/TimespanStructCustomization.cpp
--- a/Engine/Source/Editor/DetailCustomizations/Private/TimespanStructCustomization.cpp
+++ b/Engine/Source/Editor/DetailCustomizations/Private/TimespanStructCustomization.cpp
@@ -14,15 +14,16 @@
 /* IDetailCustomization interface
  *****************************************************************************/
 
-void FTimespanStructCustomization::CustomizeChildren( TSharedRef<IPropertyHandle> StructPropertyHandle, class IDetailChildrenBuilder& StructBuilder, IPropertyTypeCustomizationUtils& StructCustomizationUtils )
+void FTimespanStructCustomization::CustomizeChildren(TSharedRef<IPropertyHandle> StructPropertyHandle, class IDetailChildrenBuilder& StructBuilder, IPropertyTypeCustomizationUtils& StructCustomizationUtils)
 {
 	/* do nothing */
 }
 
 
-void FTimespanStructCustomization::CustomizeHeader( TSharedRef<IPropertyHandle> StructPropertyHandle, class FDetailWidgetRow& HeaderRow, IPropertyTypeCustomizationUtils& StructCustomizationUtils )
+void FTimespanStructCustomization::CustomizeHeader(TSharedRef<IPropertyHandle> StructPropertyHandle, class FDetailWidgetRow& HeaderRow, IPropertyTypeCustomizationUtils& StructCustomizationUtils)
 {
 	PropertyHandle = StructPropertyHandle;
+	InputValid = true;
 
 	HeaderRow
 		.NameContent()
@@ -40,7 +41,7 @@ void FTimespanStructCustomization::CustomizeHeader( TSharedRef<IPropertyHandle>
 				.OnTextChanged(this, &FTimespanStructCustomization::HandleTextBoxTextChanged)
 				.OnTextCommitted(this, &FTimespanStructCustomization::HandleTextBoxTextCommited)
 				.SelectAllTextOnCommit(true)
-				.Font( IPropertyTypeCustomizationUtils::GetRegularFont() )
+				.Font(IPropertyTypeCustomizationUtils::GetRegularFont())
 				.Text(this, &FTimespanStructCustomization::HandleTextBoxText)
 		];
 }
@@ -49,19 +50,19 @@ void FTimespanStructCustomization::CustomizeHeader( TSharedRef<IPropertyHandle>
 /* FTimespanStructCustomization callbacks
  *****************************************************************************/
 
-FSlateColor FTimespanStructCustomization::HandleTextBoxForegroundColor( ) const
+FSlateColor FTimespanStructCustomization::HandleTextBoxForegroundColor() const
 {
-	if (InputValid)
+	if (!InputValid)
 	{
-		static const FName InvertedForegroundName("InvertedForeground");
-		return FEditorStyle::GetSlateColor(InvertedForegroundName);
+		return FLinearColor::Red;
 	}
 
-	return FLinearColor::Red;
+	static const FName InvertedForegroundName("InvertedForeground");
+	return FEditorStyle::GetSlateColor(InvertedForegroundName);
 }
 
 
-FText FTimespanStructCustomization::HandleTextBoxText( ) const
+FText FTimespanStructCustomization::HandleTextBoxText() const
 {
 	TArray<void*> RawData;
 	PropertyHandle->AccessRawData(RawData);
@@ -70,39 +71,47 @@ FText FTimespanStructCustomization::HandleTextBoxText( ) const
 	{
 		return LOCTEXT("MultipleValues", "Multiple Values");
 	}
-	else if( RawData[0] == nullptr )
+	
+	if(RawData[0] == nullptr)
 	{
 		return FText::GetEmpty();
 	}
-	else
+
+	FString ValueString;
+
+	if (!((FTimespan*)RawData[0])->ExportTextItem(ValueString, FTimespan::Zero(), NULL, 0, NULL))
 	{
-		return FText::FromString(((FTimespan*)RawData[0])->ToString());
+		return FText::GetEmpty();
 	}
+	
+	return FText::FromString(ValueString);
 }
 
 
-void FTimespanStructCustomization::HandleTextBoxTextChanged( const FText& NewText )
+void FTimespanStructCustomization::HandleTextBoxTextChanged(const FText& NewText)
 {
 	FTimespan Timespan;
 	InputValid = FTimespan::Parse(NewText.ToString(), Timespan);
 }
 
 
-void FTimespanStructCustomization::HandleTextBoxTextCommited( const FText& NewText, ETextCommit::Type CommitInfo )
+void FTimespanStructCustomization::HandleTextBoxTextCommited(const FText& NewText, ETextCommit::Type CommitInfo)
 {
 	FTimespan ParsedTimespan;
-			
 	InputValid = FTimespan::Parse(NewText.ToString(), ParsedTimespan);
+
 	if (InputValid && PropertyHandle.IsValid())
 	{
 		TArray<void*> RawData;
-		PropertyHandle->AccessRawData(RawData);
 
+		PropertyHandle->AccessRawData(RawData);
 		PropertyHandle->NotifyPreChange();
+
 		for (auto RawDataInstance : RawData)
 		{
 			*(FTimespan*)RawDataInstance = ParsedTimespan;
 		}
+
 		PropertyHandle->NotifyPostChange();
 		PropertyHandle->NotifyFinishedChangingProperties();
 	}
diff --git a/Engine/Source/Editor/EditorStyle/Private/SlateEditorStyle.cpp b/Engine/Source/Editor/EditorStyle/Private/SlateEditorStyle.cpp
--- a/Engine/Source/Editor/EditorStyle/Private/SlateEditorStyle.cpp
+++ b/Engine/Source/Editor/EditorStyle/Private/SlateEditorStyle.cpp
@@ -4919,6 +4919,7 @@ void FSlateEditorStyle::FStyle::SetupLevelEditorStyle()
 		Set( "LevelViewport.Fog", new IMAGE_BRUSH( "Icons/icon_ShowFog_16x", Icon16x16 ) );
 		Set( "LevelViewport.Grid", new IMAGE_BRUSH( "Icons/icon_ShowGrid_16x", Icon16x16 ) );
 		Set( "LevelViewport.Landscape", new IMAGE_BRUSH( "Icons/icon_ShowLandscape_16x", Icon16x16 ) );
+		Set( "LevelViewport.MediaPlanes", new IMAGE_BRUSH( "Icons/icon_ShowMediaPlanes_16x", Icon16x16 ) );
 		Set( "LevelViewport.Navigation", new IMAGE_BRUSH( "Icons/icon_ShowNavigation_16x", Icon16x16 ) );
 		Set( "LevelViewport.Particles", new IMAGE_BRUSH( "Icons/icon_ShowParticlesSprite_16x", Icon16x16 ) );
 		Set( "LevelViewport.SkeletalMeshes", new IMAGE_BRUSH( "Icons/icon_ShowSkeletalMeshes_16x", Icon16x16 ) );
@@ -5825,7 +5826,6 @@ void FSlateEditorStyle::FStyle::SetupClassIconsAndThumbnails()
 			TEXT("MaterialParameterCollection"),
 			TEXT("MatineeActor"),
 			TEXT("MediaPlayer"),
-			TEXT("MediaSoundWave"),
 			TEXT("MediaTexture"),
 			TEXT("LevelSequenceActor"),
 			TEXT("MultiFont"),
diff --git a/Engine/Source/Editor/LandscapeEditor/Private/LandscapeFileFormatPng.cpp b/Engine/Source/Editor/LandscapeEditor/Private/LandscapeFileFormatPng.cpp
--- a/Engine/Source/Editor/LandscapeEditor/Private/LandscapeFileFormatPng.cpp
+++ b/Engine/Source/Editor/LandscapeEditor/Private/LandscapeFileFormatPng.cpp
@@ -3,12 +3,14 @@
 #include "LandscapeFileFormatPng.h"
 #include "Misc/FileHelper.h"
 #include "Modules/ModuleManager.h"
-
-#include "Interfaces/IImageWrapperModule.h"
+#include "IImageWrapper.h"
+#include "IImageWrapperModule.h"
 #include "Containers/Algo/Transform.h"
 
+
 #define LOCTEXT_NAMESPACE "LandscapeEditor.NewLandscape"
 
+
 FLandscapeHeightmapFileFormat_Png::FLandscapeHeightmapFileFormat_Png()
 {
 	FileTypeInfo.Description = LOCTEXT("FileFormatPng_HeightmapDesc", "Heightmap .png files");
@@ -29,7 +31,7 @@ FLandscapeHeightmapInfo FLandscapeHeightmapFileFormat_Png::Validate(const TCHAR*
 	else
 	{
 		IImageWrapperModule& ImageWrapperModule = FModuleManager::LoadModuleChecked<IImageWrapperModule>("ImageWrapper");
-		IImageWrapperPtr ImageWrapper = ImageWrapperModule.CreateImageWrapper(EImageFormat::PNG);
+		TSharedPtr<IImageWrapper> ImageWrapper = ImageWrapperModule.CreateImageWrapper(EImageFormat::PNG);
 
 		if (!ImageWrapper->SetCompressed(ImportData.GetData(), ImportData.Num()) || ImageWrapper->GetWidth() <= 0 || ImageWrapper->GetHeight() <= 0)
 		{
@@ -75,7 +77,7 @@ FLandscapeHeightmapImportData FLandscapeHeightmapFileFormat_Png::Import(const TC
 	else
 	{
 		IImageWrapperModule& ImageWrapperModule = FModuleManager::LoadModuleChecked<IImageWrapperModule>("ImageWrapper");
-		IImageWrapperPtr ImageWrapper = ImageWrapperModule.CreateImageWrapper(EImageFormat::PNG);
+		TSharedPtr<IImageWrapper> ImageWrapper = ImageWrapperModule.CreateImageWrapper(EImageFormat::PNG);
 
 		if (!ImageWrapper->SetCompressed(TempData.GetData(), TempData.Num()))
 		{
@@ -137,7 +139,7 @@ FLandscapeHeightmapImportData FLandscapeHeightmapFileFormat_Png::Import(const TC
 void FLandscapeHeightmapFileFormat_Png::Export(const TCHAR* HeightmapFilename, TArrayView<const uint16> Data, FLandscapeFileResolution DataResolution, FVector Scale) const
 {
 	IImageWrapperModule& ImageWrapperModule = FModuleManager::LoadModuleChecked<IImageWrapperModule>("ImageWrapper");
-	IImageWrapperPtr ImageWrapper = ImageWrapperModule.CreateImageWrapper(EImageFormat::PNG);
+	TSharedPtr<IImageWrapper> ImageWrapper = ImageWrapperModule.CreateImageWrapper(EImageFormat::PNG);
 
 	if (ImageWrapper->SetRaw(Data.GetData(), Data.Num() * 2, DataResolution.Width, DataResolution.Height, ERGBFormat::Gray, 16))
 	{
@@ -168,7 +170,7 @@ FLandscapeWeightmapInfo FLandscapeWeightmapFileFormat_Png::Validate(const TCHAR*
 	else
 	{
 		IImageWrapperModule& ImageWrapperModule = FModuleManager::LoadModuleChecked<IImageWrapperModule>("ImageWrapper");
-		IImageWrapperPtr ImageWrapper = ImageWrapperModule.CreateImageWrapper(EImageFormat::PNG);
+		TSharedPtr<IImageWrapper> ImageWrapper = ImageWrapperModule.CreateImageWrapper(EImageFormat::PNG);
 
 		if (!ImageWrapper->SetCompressed(ImportData.GetData(), ImportData.Num()))
 		{
@@ -205,7 +207,7 @@ FLandscapeWeightmapImportData FLandscapeWeightmapFileFormat_Png::Import(const TC
 	else
 	{
 		IImageWrapperModule& ImageWrapperModule = FModuleManager::LoadModuleChecked<IImageWrapperModule>("ImageWrapper");
-		IImageWrapperPtr ImageWrapper = ImageWrapperModule.CreateImageWrapper(EImageFormat::PNG);
+		TSharedPtr<IImageWrapper> ImageWrapper = ImageWrapperModule.CreateImageWrapper(EImageFormat::PNG);
 
 		const TArray<uint8>* RawData = nullptr;
 		if (!ImageWrapper->SetCompressed(TempData.GetData(), TempData.Num()))
@@ -241,7 +243,7 @@ FLandscapeWeightmapImportData FLandscapeWeightmapFileFormat_Png::Import(const TC
 void FLandscapeWeightmapFileFormat_Png::Export(const TCHAR* WeightmapFilename, FName LayerName, TArrayView<const uint8> Data, FLandscapeFileResolution DataResolution) const
 {
 	IImageWrapperModule& ImageWrapperModule = FModuleManager::LoadModuleChecked<IImageWrapperModule>("ImageWrapper");
-	IImageWrapperPtr ImageWrapper = ImageWrapperModule.CreateImageWrapper(EImageFormat::PNG);
+	TSharedPtr<IImageWrapper> ImageWrapper = ImageWrapperModule.CreateImageWrapper(EImageFormat::PNG);
 
 	if (ImageWrapper->SetRaw(Data.GetData(), Data.Num(), DataResolution.Width, DataResolution.Height, ERGBFormat::Gray, 8))
 	{
diff --git a/Engine/Source/Editor/LiveLinkMessageBusFramework/LiveLinkMessageBusFramework.Build.cs b/Engine/Source/Editor/LiveLinkMessageBusFramework/LiveLinkMessageBusFramework.Build.cs
--- a/Engine/Source/Editor/LiveLinkMessageBusFramework/LiveLinkMessageBusFramework.Build.cs
+++ b/Engine/Source/Editor/LiveLinkMessageBusFramework/LiveLinkMessageBusFramework.Build.cs
@@ -6,7 +6,8 @@ public class LiveLinkMessageBusFramework : ModuleRules
 {
     public LiveLinkMessageBusFramework(ReadOnlyTargetRules Target) : base(Target)
     {
-		PrivateDependencyModuleNames.Add("Messaging");
+		PrivateDependencyModuleNames.Add("MessagingCommon");
+        PrivateIncludePathModuleNames.Add("Messaging");
 		
         PublicDependencyModuleNames.AddRange(
             new string[] {
diff --git a/Engine/Source/Editor/LiveLinkMessageBusFramework/Private/LiveLinkProvider.cpp b/Engine/Source/Editor/LiveLinkMessageBusFramework/Private/LiveLinkProvider.cpp
--- a/Engine/Source/Editor/LiveLinkMessageBusFramework/Private/LiveLinkProvider.cpp
+++ b/Engine/Source/Editor/LiveLinkMessageBusFramework/Private/LiveLinkProvider.cpp
@@ -50,7 +50,7 @@ private:
 	
 	const FString MachineName;
 
-	FMessageEndpointPtr MessageEndpoint;
+	TSharedPtr<FMessageEndpoint, ESPMode::ThreadSafe> MessageEndpoint;
 
 	// Array of our current connections
 	TArray<FTrackedAddress> ConnectedAddresses;
@@ -59,9 +59,9 @@ private:
 	TMap<FName, FTrackedSubject> Subjects;
 	
 	//Message bus message handlers
-	void HandlePingMessage(const FLiveLinkPingMessage& Message, const IMessageContextRef& Context);
-	void HandleConnectMessage(const FLiveLinkConnectMessage& Message, const IMessageContextRef& Context);
-	void HandleHeartbeat(const FLiveLinkHeartbeatMessage& Message, const IMessageContextRef& Context);
+	void HandlePingMessage(const FLiveLinkPingMessage& Message, const TSharedRef<IMessageContext, ESPMode::ThreadSafe>& Context);
+	void HandleConnectMessage(const FLiveLinkConnectMessage& Message, const TSharedRef<IMessageContext, ESPMode::ThreadSafe>& Context);
+	void HandleHeartbeat(const FLiveLinkHeartbeatMessage& Message, const TSharedRef<IMessageContext, ESPMode::ThreadSafe>& Context);
 	// End message bus message handlers
 
 	// Validate our current connections
@@ -210,12 +210,12 @@ public:
 
 const double FLiveLinkProvider::CONNECTION_TIMEOUT = 10.f;
 
-void FLiveLinkProvider::HandlePingMessage(const FLiveLinkPingMessage& Message, const IMessageContextRef& Context)
+void FLiveLinkProvider::HandlePingMessage(const FLiveLinkPingMessage& Message, const TSharedRef<IMessageContext, ESPMode::ThreadSafe>& Context)
 {
 	MessageEndpoint->Send(new FLiveLinkPongMessage(ProviderName, MachineName, Message.PollRequest), Context->GetSender());
 }
 
-void FLiveLinkProvider::HandleConnectMessage(const FLiveLinkConnectMessage& Message, const IMessageContextRef& Context)
+void FLiveLinkProvider::HandleConnectMessage(const FLiveLinkConnectMessage& Message, const TSharedRef<IMessageContext, ESPMode::ThreadSafe>& Context)
 {
 	const FMessageAddress& ConnectionAddress = Context->GetSender();
 
@@ -231,7 +231,7 @@ void FLiveLinkProvider::HandleConnectMessage(const FLiveLinkConnectMessage& Mess
 	}
 }
 
-void FLiveLinkProvider::HandleHeartbeat(const FLiveLinkHeartbeatMessage& Message, const IMessageContextRef& Context)
+void FLiveLinkProvider::HandleHeartbeat(const FLiveLinkHeartbeatMessage& Message, const TSharedRef<IMessageContext, ESPMode::ThreadSafe>& Context)
 {
 	FTrackedAddress* TrackedAddress = ConnectedAddresses.FindByPredicate([=](const FTrackedAddress& ConAddress) { return ConAddress.Address == Context->GetSender(); });
 	if (TrackedAddress)
diff --git a/Engine/Source/Editor/MovieSceneTools/Private/TrackEditors/SpawnTrackEditor.cpp b/Engine/Source/Editor/MovieSceneTools/Private/TrackEditors/SpawnTrackEditor.cpp
--- a/Engine/Source/Editor/MovieSceneTools/Private/TrackEditors/SpawnTrackEditor.cpp
+++ b/Engine/Source/Editor/MovieSceneTools/Private/TrackEditors/SpawnTrackEditor.cpp
@@ -60,6 +60,12 @@ bool FSpawnTrackEditor::SupportsType(TSubclassOf<UMovieSceneTrack> Type) const
 }
 
 
+bool FSpawnTrackEditor::SupportsSequence(UMovieSceneSequence* InSequence) const
+{
+	return (InSequence != nullptr) && (InSequence->GetClass()->GetName() == TEXT("LevelSequence"));
+}
+
+
 TSharedRef<ISequencerSection> FSpawnTrackEditor::MakeSectionInterface(UMovieSceneSection& SectionObject, UMovieSceneTrack& Track, FGuid ObjectBinding)
 {
 	return MakeShareable(new FBoolPropertySection(SectionObject, Track.GetDisplayName()));
diff --git a/Engine/Source/Editor/MovieSceneTools/Public/MovieSceneToolsUserSettings.h b/Engine/Source/Editor/MovieSceneTools/Public/MovieSceneToolsUserSettings.h
--- a/Engine/Source/Editor/MovieSceneTools/Public/MovieSceneToolsUserSettings.h
+++ b/Engine/Source/Editor/MovieSceneTools/Public/MovieSceneToolsUserSettings.h
@@ -27,7 +27,7 @@ public:
 	UPROPERTY(EditAnywhere, config, Category=General)
 	bool bDrawThumbnails;
 
-	/** Whether to draw thumbnails or not */
+	/** Whether to draw a single thumbnail for this section or as many as can fit */
 	UPROPERTY(EditAnywhere, config, Category=General, meta=(EditCondition=bDrawThumbnails))
 	bool bDrawSingleThumbnails;
 
diff --git a/Engine/Source/Editor/MovieSceneTools/Public/PropertyTrackEditor.h b/Engine/Source/Editor/MovieSceneTools/Public/PropertyTrackEditor.h
--- a/Engine/Source/Editor/MovieSceneTools/Public/PropertyTrackEditor.h
+++ b/Engine/Source/Editor/MovieSceneTools/Public/PropertyTrackEditor.h
@@ -97,7 +97,12 @@ public:
 
 public:
 
-	// ISequencerTrackEditor interface
+	//~ ISequencerTrackEditor interface
+
+    virtual bool SupportsSequence(UMovieSceneSequence* InSequence) const override
+    {
+        return true;
+    }
 
 	virtual bool SupportsType( TSubclassOf<UMovieSceneTrack> Type ) const override
 	{
diff --git a/Engine/Source/Editor/MovieSceneTools/Public/TrackEditors/AudioTrackEditor.h b/Engine/Source/Editor/MovieSceneTools/Public/TrackEditors/AudioTrackEditor.h
--- a/Engine/Source/Editor/MovieSceneTools/Public/TrackEditors/AudioTrackEditor.h
+++ b/Engine/Source/Editor/MovieSceneTools/Public/TrackEditors/AudioTrackEditor.h
@@ -70,7 +70,7 @@ protected:
 
 private:
 
-	/** Callback for executing the "Add Event Track" menu entry. */
+	/** Callback for executing the "Add Audio Track" menu entry. */
 	void HandleAddAudioTrackMenuEntryExecute();
 
 	/** Audio sub menu */
diff --git a/Engine/Source/Editor/SwarmInterface/Private/SwarmInterfaceLocal.cpp b/Engine/Source/Editor/SwarmInterface/Private/SwarmInterfaceLocal.cpp
--- a/Engine/Source/Editor/SwarmInterface/Private/SwarmInterfaceLocal.cpp
+++ b/Engine/Source/Editor/SwarmInterface/Private/SwarmInterfaceLocal.cpp
@@ -10,14 +10,15 @@
 #include "Async/TaskGraphInterfaces.h"
 
 #include "SwarmInterface.h"
-
 #include "SwarmMessages.h"
+
 #if USE_LOCAL_SWARM_INTERFACE
-#include "IMessageContext.h"
-#include "Helpers/MessageEndpoint.h"
-#include "Helpers/MessageEndpointBuilder.h"
+	#include "IMessageContext.h"
+	#include "MessageEndpoint.h"
+	#include "MessageEndpointBuilder.h"
 #endif
 
+
 namespace NSwarm
 {
 
@@ -54,17 +55,17 @@ private:
 	void PrepareTasksList();
 
 #if USE_LOCAL_SWARM_INTERFACE
-	void HandlePingMessage( const FSwarmPingMessage& Message, const IMessageContextRef& Context );
-	void HandlePongMessage( const FSwarmPongMessage& Message, const IMessageContextRef& Context );
-	void HandleInfoMessage( const FSwarmInfoMessage& Message, const IMessageContextRef& Context );
-	void HandleAlertMessage( const FSwarmAlertMessage& Message, const IMessageContextRef& Context );
-	void HandleTimingMessage( const FSwarmTimingMessage& Message, const IMessageContextRef& Context );
-	void HandleTaskRequestReleaseMessage( const FSwarmTaskRequestReleaseMessage& Message, const IMessageContextRef& Context );
-	void HandleTaskRequestReservationMessage( const FSwarmTaskRequestReservationMessage& Message, const IMessageContextRef& Context );
-	void HandleTaskRequestSpecificationMessage( const FSwarmTaskRequestSpecificationMessage& Message, const IMessageContextRef& Context );
-	void HandleJobStateMessage( const FSwarmJobStateMessage& Message, const IMessageContextRef& Context );
-	void HandleTaskStateMessage( const FSwarmTaskStateMessage& Message, const IMessageContextRef& Context );
-	void HandleQuitMessage( const FSwarmQuitMessage& Message, const IMessageContextRef& Context );
+	void HandlePingMessage( const FSwarmPingMessage& Message, const TSharedRef<IMessageContext, ESPMode::ThreadSafe>& Context );
+	void HandlePongMessage( const FSwarmPongMessage& Message, const TSharedRef<IMessageContext, ESPMode::ThreadSafe>& Context );
+	void HandleInfoMessage( const FSwarmInfoMessage& Message, const TSharedRef<IMessageContext, ESPMode::ThreadSafe>& Context );
+	void HandleAlertMessage( const FSwarmAlertMessage& Message, const TSharedRef<IMessageContext, ESPMode::ThreadSafe>& Context );
+	void HandleTimingMessage( const FSwarmTimingMessage& Message, const TSharedRef<IMessageContext, ESPMode::ThreadSafe>& Context );
+	void HandleTaskRequestReleaseMessage( const FSwarmTaskRequestReleaseMessage& Message, const TSharedRef<IMessageContext, ESPMode::ThreadSafe>& Context );
+	void HandleTaskRequestReservationMessage( const FSwarmTaskRequestReservationMessage& Message, const TSharedRef<IMessageContext, ESPMode::ThreadSafe>& Context );
+	void HandleTaskRequestSpecificationMessage( const FSwarmTaskRequestSpecificationMessage& Message, const TSharedRef<IMessageContext, ESPMode::ThreadSafe>& Context );
+	void HandleJobStateMessage( const FSwarmJobStateMessage& Message, const TSharedRef<IMessageContext, ESPMode::ThreadSafe>& Context );
+	void HandleTaskStateMessage( const FSwarmTaskStateMessage& Message, const TSharedRef<IMessageContext, ESPMode::ThreadSafe>& Context );
+	void HandleQuitMessage( const FSwarmQuitMessage& Message, const TSharedRef<IMessageContext, ESPMode::ThreadSafe>& Context );
 #endif
 
 	FString JobFolder;
@@ -75,7 +76,7 @@ private:
 	FConnectionCallback CallbackFunc;
 	void* CallbackData;
 #if USE_LOCAL_SWARM_INTERFACE
-	FMessageEndpointPtr MessageEndpoint;
+	TSharedPtr<FMessageEndpoint, ESPMode::ThreadSafe> MessageEndpoint;
 	FMessageAddress Recepient;
 	bool bIsConnected;
 	bool bIsEditor;
@@ -269,12 +270,12 @@ int32 FSwarmInterfaceLocalImpl::SendMessage( const FMessage& Message )
 }
 
 #if USE_LOCAL_SWARM_INTERFACE
-void FSwarmInterfaceLocalImpl::HandlePingMessage( const FSwarmPingMessage& Message, const IMessageContextRef& Context )
+void FSwarmInterfaceLocalImpl::HandlePingMessage( const FSwarmPingMessage& Message, const TSharedRef<IMessageContext, ESPMode::ThreadSafe>& Context )
 {
 	MessageEndpoint->Send(new FSwarmPongMessage(bIsEditor, FPlatformProcess::ComputerName()), Context->GetSender());
 }
 
-void FSwarmInterfaceLocalImpl::HandlePongMessage( const FSwarmPongMessage& Message, const IMessageContextRef& Context )
+void FSwarmInterfaceLocalImpl::HandlePongMessage( const FSwarmPongMessage& Message, const TSharedRef<IMessageContext, ESPMode::ThreadSafe>& Context )
 {
 	if (!Recepient.IsValid() && Message.bIsEditor != bIsEditor && Message.ComputerName == FPlatformProcess::ComputerName())
 	{
@@ -282,7 +283,7 @@ void FSwarmInterfaceLocalImpl::HandlePongMessage( const FSwarmPongMessage& Messa
 	}
 }
 	
-void FSwarmInterfaceLocalImpl::HandleInfoMessage( const FSwarmInfoMessage& Message, const IMessageContextRef& Context )
+void FSwarmInterfaceLocalImpl::HandleInfoMessage( const FSwarmInfoMessage& Message, const TSharedRef<IMessageContext, ESPMode::ThreadSafe>& Context )
 {
 	if( CallbackFunc )
 	{
@@ -291,7 +292,7 @@ void FSwarmInterfaceLocalImpl::HandleInfoMessage( const FSwarmInfoMessage& Messa
 	}
 }
 
-void FSwarmInterfaceLocalImpl::HandleAlertMessage( const FSwarmAlertMessage& Message, const IMessageContextRef& Context )
+void FSwarmInterfaceLocalImpl::HandleAlertMessage( const FSwarmAlertMessage& Message, const TSharedRef<IMessageContext, ESPMode::ThreadSafe>& Context )
 {
 	if( CallbackFunc )
 	{
@@ -300,7 +301,7 @@ void FSwarmInterfaceLocalImpl::HandleAlertMessage( const FSwarmAlertMessage& Mes
 	}
 }
 
-void FSwarmInterfaceLocalImpl::HandleTimingMessage( const FSwarmTimingMessage& Message, const IMessageContextRef& Context )
+void FSwarmInterfaceLocalImpl::HandleTimingMessage( const FSwarmTimingMessage& Message, const TSharedRef<IMessageContext, ESPMode::ThreadSafe>& Context )
 {
 	if( CallbackFunc )
 	{
@@ -309,7 +310,7 @@ void FSwarmInterfaceLocalImpl::HandleTimingMessage( const FSwarmTimingMessage& M
 	}
 }
 
-void FSwarmInterfaceLocalImpl::HandleTaskRequestReleaseMessage( const FSwarmTaskRequestReleaseMessage& Message, const IMessageContextRef& Context )
+void FSwarmInterfaceLocalImpl::HandleTaskRequestReleaseMessage( const FSwarmTaskRequestReleaseMessage& Message, const TSharedRef<IMessageContext, ESPMode::ThreadSafe>& Context )
 {
 	if( CallbackFunc )
 	{
@@ -318,7 +319,7 @@ void FSwarmInterfaceLocalImpl::HandleTaskRequestReleaseMessage( const FSwarmTask
 	}
 }
 
-void FSwarmInterfaceLocalImpl::HandleTaskRequestReservationMessage( const FSwarmTaskRequestReservationMessage& Message, const IMessageContextRef& Context )
+void FSwarmInterfaceLocalImpl::HandleTaskRequestReservationMessage( const FSwarmTaskRequestReservationMessage& Message, const TSharedRef<IMessageContext, ESPMode::ThreadSafe>& Context )
 {
 	if( CallbackFunc )
 	{
@@ -327,7 +328,7 @@ void FSwarmInterfaceLocalImpl::HandleTaskRequestReservationMessage( const FSwarm
 	}
 }
 
-void FSwarmInterfaceLocalImpl::HandleTaskRequestSpecificationMessage( const FSwarmTaskRequestSpecificationMessage& Message, const IMessageContextRef& Context )
+void FSwarmInterfaceLocalImpl::HandleTaskRequestSpecificationMessage( const FSwarmTaskRequestSpecificationMessage& Message, const TSharedRef<IMessageContext, ESPMode::ThreadSafe>& Context )
 {
 	if( CallbackFunc )
 	{
@@ -350,7 +351,7 @@ void FSwarmInterfaceLocalImpl::HandleTaskRequestSpecificationMessage( const FSwa
 	}
 }
 
-void FSwarmInterfaceLocalImpl::HandleJobStateMessage( const FSwarmJobStateMessage& Message, const IMessageContextRef& Context )
+void FSwarmInterfaceLocalImpl::HandleJobStateMessage( const FSwarmJobStateMessage& Message, const TSharedRef<IMessageContext, ESPMode::ThreadSafe>& Context )
 {
 	if( CallbackFunc )
 	{
@@ -359,7 +360,7 @@ void FSwarmInterfaceLocalImpl::HandleJobStateMessage( const FSwarmJobStateMessag
 	}
 }
 
-void FSwarmInterfaceLocalImpl::HandleTaskStateMessage( const FSwarmTaskStateMessage& Message, const IMessageContextRef& Context )
+void FSwarmInterfaceLocalImpl::HandleTaskStateMessage( const FSwarmTaskStateMessage& Message, const TSharedRef<IMessageContext, ESPMode::ThreadSafe>& Context )
 {
 	if( CallbackFunc )
 	{
@@ -368,7 +369,7 @@ void FSwarmInterfaceLocalImpl::HandleTaskStateMessage( const FSwarmTaskStateMess
 	}
 }
 
-void FSwarmInterfaceLocalImpl::HandleQuitMessage( const FSwarmQuitMessage& Message, const IMessageContextRef& Context )
+void FSwarmInterfaceLocalImpl::HandleQuitMessage( const FSwarmQuitMessage& Message, const TSharedRef<IMessageContext, ESPMode::ThreadSafe>& Context )
 {
 	if( CallbackFunc )
 	{
diff --git a/Engine/Source/Editor/SwarmInterface/SwarmInterface.Build.cs b/Engine/Source/Editor/SwarmInterface/SwarmInterface.Build.cs
--- a/Engine/Source/Editor/SwarmInterface/SwarmInterface.Build.cs
+++ b/Engine/Source/Editor/SwarmInterface/SwarmInterface.Build.cs
@@ -21,7 +21,7 @@ public class SwarmInterface : ModuleRules
 
 		if (Target.Platform == UnrealTargetPlatform.Mac || Target.Platform == UnrealTargetPlatform.Linux)
 		{
-			PrivateDependencyModuleNames.Add("Messaging");
+			PrivateIncludePathModuleNames.Add("MessagingCommon");
 		}
 	}
 }
diff --git a/Engine/Source/Editor/UnrealEd/Classes/CookOnTheSide/CookOnTheFlyServer.h b/Engine/Source/Editor/UnrealEd/Classes/CookOnTheSide/CookOnTheFlyServer.h
--- a/Engine/Source/Editor/UnrealEd/Classes/CookOnTheSide/CookOnTheFlyServer.h
+++ b/Engine/Source/Editor/UnrealEd/Classes/CookOnTheSide/CookOnTheFlyServer.h
@@ -1,9 +1,5 @@
 // Copyright 1998-2017 Epic Games, Inc. All Rights Reserved.
 
-/*=============================================================================
-CookOnTheFlyServer.h : handles polite cook requests via network ;)
-=============================================================================*/
-
 #pragma once
 
 #include "CoreMinimal.h"
@@ -18,7 +14,7 @@ CookOnTheFlyServer.h : handles polite cook requests via network ;)
 #include "HAL/PlatformProcess.h"
 #include "TickableEditorObject.h"
 #include "IPlatformFileSandboxWrapper.h"
-#include "Interfaces/INetworkFileSystemModule.h"
+#include "INetworkFileSystemModule.h"
 #include "CookOnTheFlyServer.generated.h"
 
 
diff --git a/Engine/Source/Editor/UnrealEd/Classes/Editor/PropertyEditorTestObject.h b/Engine/Source/Editor/UnrealEd/Classes/Editor/PropertyEditorTestObject.h
--- a/Engine/Source/Editor/UnrealEd/Classes/Editor/PropertyEditorTestObject.h
+++ b/Engine/Source/Editor/UnrealEd/Classes/Editor/PropertyEditorTestObject.h
@@ -166,13 +166,16 @@ class UPropertyEditorTestObject : public UObject
 	UPROPERTY(EditAnywhere, Category=BasicProperties)
 	FText TextProperty;
 
-	UPROPERTY(EditAnywhere, Category=BasicProperties)
+	UPROPERTY(EditAnywhere, Category=BasicProperties, meta=(AllowPreserveRatio))
+	FIntPoint IntPointProperty;
+
+	UPROPERTY(EditAnywhere, Category=BasicProperties, meta=(AllowPreserveRatio))
 	FVector Vector3Property;
 
-	UPROPERTY(EditAnywhere, Category=BasicProperties)
+	UPROPERTY(EditAnywhere, Category=BasicProperties, meta=(AllowPreserveRatio))
 	FVector2D Vector2Property;
 
-	UPROPERTY(EditAnywhere, Category=BasicProperties)
+	UPROPERTY(EditAnywhere, Category=BasicProperties, meta=(AllowPreserveRatio))
 	FVector4 Vector4Property;
 
 	UPROPERTY(EditAnywhere, Category=BasicProperties)
@@ -193,6 +196,12 @@ class UPropertyEditorTestObject : public UObject
 	UPROPERTY(EditAnywhere, Category=BasicProperties)
 	TEnumAsByte<enum PropertEditorTestEnum> EnumProperty;
 
+	UPROPERTY(EditAnywhere, Category = BasicProperties)
+	FMatrix MatrixProperty;
+
+	UPROPERTY(EditAnywhere, Category = BasicProperties)
+	FTransform TransformProperty;
+
 	// Integer
 	UPROPERTY(EditAnywhere, Category=ArraysOfProperties)
 	TArray<int32> IntProperty32Array;
diff --git a/Engine/Source/Editor/UnrealEd/Private/AutoReimport/ContentDirectoryMonitor.cpp b/Engine/Source/Editor/UnrealEd/Private/AutoReimport/ContentDirectoryMonitor.cpp
--- a/Engine/Source/Editor/UnrealEd/Private/AutoReimport/ContentDirectoryMonitor.cpp
+++ b/Engine/Source/Editor/UnrealEd/Private/AutoReimport/ContentDirectoryMonitor.cpp
@@ -123,7 +123,7 @@ void FContentDirectoryMonitor::Tick()
 	Cache.Tick();
 
 	// Immediately resolve any changes that we should not consider
-	const FDateTime Threshold = FDateTime::UtcNow() - FTimespan(0, 0, GetDefault<UEditorLoadingSavingSettings>()->AutoReimportThreshold);
+	const FDateTime Threshold = FDateTime::UtcNow() - FTimespan::FromSeconds(GetDefault<UEditorLoadingSavingSettings>()->AutoReimportThreshold);
 
 	TArray<DirectoryWatcher::FUpdateCacheTransaction> InsignificantTransactions = Cache.FilterOutstandingChanges([=](const DirectoryWatcher::FUpdateCacheTransaction& Transaction, const FDateTime& TimeOfChange){
 		return TimeOfChange <= Threshold && !ShouldConsiderChange(Transaction);
@@ -154,7 +154,7 @@ bool FContentDirectoryMonitor::ShouldConsiderChange(const DirectoryWatcher::FUpd
 
 int32 FContentDirectoryMonitor::GetNumUnprocessedChanges() const
 {
-	const FDateTime Threshold = FDateTime::UtcNow() - FTimespan(0, 0, GetDefault<UEditorLoadingSavingSettings>()->AutoReimportThreshold);
+	const FDateTime Threshold = FDateTime::UtcNow() - FTimespan::FromSeconds(GetDefault<UEditorLoadingSavingSettings>()->AutoReimportThreshold);
 
 	int32 Total = 0;
 
@@ -179,7 +179,7 @@ int32 FContentDirectoryMonitor::StartProcessing()
 {
 	// We only process things that haven't changed for a given threshold
 	auto& FileManager = IFileManager::Get();
-	const FDateTime Threshold = FDateTime::UtcNow() - FTimespan(0, 0, GetDefault<UEditorLoadingSavingSettings>()->AutoReimportThreshold);
+	const FDateTime Threshold = FDateTime::UtcNow() - FTimespan::FromSeconds(GetDefault<UEditorLoadingSavingSettings>()->AutoReimportThreshold);
 
 	// Get all the changes that have happend beyond our import threshold
 	auto OutstandingChanges = Cache.FilterOutstandingChanges([=](const DirectoryWatcher::FUpdateCacheTransaction& Transaction, const FDateTime& TimeOfChange){
diff --git a/Engine/Source/Editor/UnrealEd/Private/Commandlets/CookCommandlet.cpp b/Engine/Source/Editor/UnrealEd/Private/Commandlets/CookCommandlet.cpp
--- a/Engine/Source/Editor/UnrealEd/Private/Commandlets/CookCommandlet.cpp
+++ b/Engine/Source/Editor/UnrealEd/Private/Commandlets/CookCommandlet.cpp
@@ -34,7 +34,7 @@
 #include "ShaderCompiler.h"
 #include "Interfaces/ITargetPlatform.h"
 #include "Interfaces/ITargetPlatformManagerModule.h"
-#include "Interfaces/INetworkFileSystemModule.h"
+#include "INetworkFileSystemModule.h"
 #include "GameDelegates.h"
 #include "CookerSettings.h"
 #include "ShaderCompiler.h"
diff --git a/Engine/Source/Editor/UnrealEd/Private/Commandlets/FileServerCommandlet.cpp b/Engine/Source/Editor/UnrealEd/Private/Commandlets/FileServerCommandlet.cpp
--- a/Engine/Source/Editor/UnrealEd/Private/Commandlets/FileServerCommandlet.cpp
+++ b/Engine/Source/Editor/UnrealEd/Private/Commandlets/FileServerCommandlet.cpp
@@ -1,9 +1,5 @@
 // Copyright 1998-2017 Epic Games, Inc. All Rights Reserved.
 
-/*=============================================================================
-	FileServerCommandlet.cpp: Implements the UFileServerCommandlet class.
-=============================================================================*/
-
 #include "Commandlets/FileServerCommandlet.h"
 #include "Misc/MessageDialog.h"
 #include "Misc/App.h"
@@ -16,10 +12,10 @@
 #include "IDirectoryWatcher.h"
 #include "DirectoryWatcherModule.h"
 #include "IMessageContext.h"
-#include "Helpers/MessageEndpoint.h"
-#include "Helpers/MessageEndpointBuilder.h"
-#include "Interfaces/INetworkFileServer.h"
-#include "Interfaces/INetworkFileSystemModule.h"
+#include "INetworkFileServer.h"
+#include "INetworkFileSystemModule.h"
+#include "MessageEndpoint.h"
+#include "MessageEndpointBuilder.h"
 #include "IPAddress.h"
 
 
@@ -90,7 +86,7 @@ int32 UFileServerCommandlet::Main( const FString& Params )
 			AddressStringList.Add(AddressList[AddressIndex]->ToString(true));
 		}
 
-		FMessageEndpointPtr MessageEndpoint = FMessageEndpoint::Builder("UFileServerCommandlet").Build();
+		TSharedPtr<FMessageEndpoint, ESPMode::ThreadSafe> MessageEndpoint = FMessageEndpoint::Builder("UFileServerCommandlet").Build();
 
 		if (MessageEndpoint.IsValid())
 		{
diff --git a/Engine/Source/Editor/UnrealEd/Private/CookOnTheFlyServer.cpp b/Engine/Source/Editor/UnrealEd/Private/CookOnTheFlyServer.cpp
--- a/Engine/Source/Editor/UnrealEd/Private/CookOnTheFlyServer.cpp
+++ b/Engine/Source/Editor/UnrealEd/Private/CookOnTheFlyServer.cpp
@@ -53,10 +53,10 @@
 #include "Interfaces/IShaderFormat.h"
 #include "Interfaces/ITextureFormat.h"
 #include "IMessageContext.h"
-#include "Helpers/MessageEndpoint.h"
-#include "Helpers/MessageEndpointBuilder.h"
-#include "Interfaces/INetworkFileServer.h"
-#include "Interfaces/INetworkFileSystemModule.h"
+#include "MessageEndpoint.h"
+#include "MessageEndpointBuilder.h"
+#include "INetworkFileServer.h"
+#include "INetworkFileSystemModule.h"
 #include "PlatformInfo.h"
 
 #include "AssetRegistryModule.h"
@@ -1032,9 +1032,7 @@ bool UCookOnTheFlyServer::BroadcastFileserverPresence( const FGuid &InstanceId )
 		}
 	}
 
-
-
-	FMessageEndpointPtr MessageEndpoint = FMessageEndpoint::Builder("UCookOnTheFlyServer").Build();
+	TSharedPtr<FMessageEndpoint, ESPMode::ThreadSafe> MessageEndpoint = FMessageEndpoint::Builder("UCookOnTheFlyServer").Build();
 
 	if (MessageEndpoint.IsValid())
 	{
diff --git a/Engine/Source/Editor/UnrealEd/Private/Factories/EditorFactories.cpp b/Engine/Source/Editor/UnrealEd/Private/Factories/EditorFactories.cpp
--- a/Engine/Source/Editor/UnrealEd/Private/Factories/EditorFactories.cpp
+++ b/Engine/Source/Editor/UnrealEd/Private/Factories/EditorFactories.cpp
@@ -186,7 +186,8 @@
 #include "DDSLoader.h"
 #include "Factories/HDRLoader.h"
 #include "Factories/IESLoader.h"
-#include "Interfaces/IImageWrapperModule.h"
+#include "IImageWrapper.h"
+#include "IImageWrapperModule.h"
 
 #include "FbxImporter.h"
 #include "FbxErrors.h"
@@ -204,16 +205,17 @@
 #include "InstancedFoliageActor.h"
 
 #if PLATFORM_WINDOWS
-// Needed for DDS support.
-#include "WindowsHWrapper.h"
-#include "AllowWindowsPlatformTypes.h"
-	#include <ddraw.h>
-#include "HideWindowsPlatformTypes.h"
+	// Needed for DDS support.
+	#include "WindowsHWrapper.h"
+	#include "AllowWindowsPlatformTypes.h"
+		#include <ddraw.h>
+	#include "HideWindowsPlatformTypes.h"
 #endif
 
 #if WITH_EDITOR
-#include "CubemapUnwrapUtils.h"
+	#include "CubemapUnwrapUtils.h"
 #endif
+
 #include "Components/BrushComponent.h"
 #include "EngineUtils.h"
 #include "Engine/AssetUserData.h"
@@ -2821,7 +2823,7 @@ UTexture* UTextureFactory::ImportTexture(UClass* Class, UObject* InParent, FName
 	//
 	// PNG
 	//
-	IImageWrapperPtr PngImageWrapper = ImageWrapperModule.CreateImageWrapper( EImageFormat::PNG );
+	TSharedPtr<IImageWrapper> PngImageWrapper = ImageWrapperModule.CreateImageWrapper( EImageFormat::PNG );
 	if ( PngImageWrapper.IsValid() && PngImageWrapper->SetCompressed( Buffer, Length ) )
 	{
 		if ( !IsImportResolutionValid( PngImageWrapper->GetWidth(), PngImageWrapper->GetHeight(), bAllowNonPowerOfTwo, Warn ) )
@@ -2832,7 +2834,8 @@ UTexture* UTextureFactory::ImportTexture(UClass* Class, UObject* InParent, FName
 		// Select the texture's source format
 		ETextureSourceFormat TextureFormat = TSF_Invalid;
 		int32 BitDepth = PngImageWrapper->GetBitDepth();
-		ERGBFormat::Type Format = PngImageWrapper->GetFormat();
+		ERGBFormat Format = PngImageWrapper->GetFormat();
+
 		if (Format == ERGBFormat::Gray)
 		{
 			if (BitDepth <= 8)
@@ -2906,7 +2909,7 @@ UTexture* UTextureFactory::ImportTexture(UClass* Class, UObject* InParent, FName
 	//
 	// JPEG
 	//
-	IImageWrapperPtr JpegImageWrapper = ImageWrapperModule.CreateImageWrapper( EImageFormat::JPEG );
+	TSharedPtr<IImageWrapper> JpegImageWrapper = ImageWrapperModule.CreateImageWrapper( EImageFormat::JPEG );
 	if ( JpegImageWrapper.IsValid() && JpegImageWrapper->SetCompressed( Buffer, Length ) )
 	{
 		if ( !IsImportResolutionValid( JpegImageWrapper->GetWidth(), JpegImageWrapper->GetHeight(), bAllowNonPowerOfTwo, Warn ) )
@@ -2917,7 +2920,7 @@ UTexture* UTextureFactory::ImportTexture(UClass* Class, UObject* InParent, FName
 		// Select the texture's source format
 		ETextureSourceFormat TextureFormat = TSF_Invalid;
 		int32 BitDepth = JpegImageWrapper->GetBitDepth();
-		ERGBFormat::Type Format = JpegImageWrapper->GetFormat();
+		ERGBFormat Format = JpegImageWrapper->GetFormat();
 
 		if ( Format == ERGBFormat::Gray )
 		{
@@ -2977,7 +2980,7 @@ UTexture* UTextureFactory::ImportTexture(UClass* Class, UObject* InParent, FName
 	//
 	// EXR
 	//
-	IImageWrapperPtr ExrImageWrapper = ImageWrapperModule.CreateImageWrapper( EImageFormat::EXR );
+	TSharedPtr<IImageWrapper> ExrImageWrapper = ImageWrapperModule.CreateImageWrapper( EImageFormat::EXR );
 	if ( ExrImageWrapper.IsValid() && ExrImageWrapper->SetCompressed( Buffer, Length ) )
 	{
 		int32 Width = ExrImageWrapper->GetWidth();
@@ -2991,7 +2994,7 @@ UTexture* UTextureFactory::ImportTexture(UClass* Class, UObject* InParent, FName
 		// Select the texture's source format
 		ETextureSourceFormat TextureFormat = TSF_Invalid;
 		int32 BitDepth = ExrImageWrapper->GetBitDepth();
-		ERGBFormat::Type Format = ExrImageWrapper->GetFormat();
+		ERGBFormat Format = ExrImageWrapper->GetFormat();
 
 		if ( Format == ERGBFormat::RGBA && BitDepth == 16 )
 		{
@@ -3039,7 +3042,7 @@ UTexture* UTextureFactory::ImportTexture(UClass* Class, UObject* InParent, FName
 	//
 	// BMP
 	//
-	IImageWrapperPtr BmpImageWrapper = ImageWrapperModule.CreateImageWrapper(EImageFormat::BMP);
+	TSharedPtr<IImageWrapper> BmpImageWrapper = ImageWrapperModule.CreateImageWrapper(EImageFormat::BMP);
 	if (BmpImageWrapper.IsValid() && BmpImageWrapper->SetCompressed(Buffer, Length))
 	{
 		// Check the resolution of the imported texture to ensure validity
diff --git a/Engine/Source/Editor/UnrealEd/Private/FileHelpers.cpp b/Engine/Source/Editor/UnrealEd/Private/FileHelpers.cpp
--- a/Engine/Source/Editor/UnrealEd/Private/FileHelpers.cpp
+++ b/Engine/Source/Editor/UnrealEd/Private/FileHelpers.cpp
@@ -964,8 +964,37 @@ void FEditorFileUtils::SaveAssetsAs(const TArray<UObject*>& Assets, TArray<UObje
 	for (UObject* Asset : Assets)
 	{
 		const FString OldPackageName = Asset->GetOutermost()->GetName();
-		const FString OldPackagePath = FPackageName::GetLongPackagePath(OldPackageName);
-		const FString OldAssetName = FPackageName::GetLongPackageAssetName(OldPackageName);
+		
+		FString OldPackagePath;
+		FString OldAssetName;
+		
+		if (Asset->HasAnyFlags(RF_Transient))
+		{
+			// determine default package path
+			const FString DefaultDirectory = FEditorDirectories::Get().GetLastDirectory(ELastDirectory::NEW_ASSET);
+			FPackageName::TryConvertFilenameToLongPackageName(DefaultDirectory, OldPackagePath);
+
+			if (OldPackagePath.IsEmpty())
+			{
+				OldPackagePath = TEXT("/Game");
+			}
+
+			// determine default asset name
+			FString DefaultName = FString(NSLOCTEXT("UnrealEd", "PrefixNew", "New").ToString() + Asset->GetClass()->GetName());
+
+			FString UniquePackageName;
+			FString UniqueAssetName;
+
+			FAssetToolsModule& AssetToolsModule = FModuleManager::LoadModuleChecked<FAssetToolsModule>("AssetTools");
+			AssetToolsModule.Get().CreateUniqueAssetName(OldPackagePath / DefaultName, TEXT(""), UniquePackageName, UniqueAssetName);
+
+			OldAssetName = FPaths::GetCleanFilename(UniqueAssetName);
+		}
+		else
+		{
+			OldAssetName = FPackageName::GetLongPackageAssetName(OldPackageName);
+			OldPackagePath = FPackageName::GetLongPackagePath(OldPackageName);
+		}
 
 		FString NewPackageName;
 
@@ -986,7 +1015,7 @@ void FEditorFileUtils::SaveAssetsAs(const TArray<UObject*>& Assets, TArray<UObje
 		// process asset
 		if (NewPackageName.IsEmpty())
 		{
-			OutSavedAssets.Add(Asset); // user cancelled
+			OutSavedAssets.Add(Asset); // user canceled
 		}
 		else if (NewPackageName != OldPackageName)
 		{
@@ -997,9 +1026,22 @@ void FEditorFileUtils::SaveAssetsAs(const TArray<UObject*>& Assets, TArray<UObje
 
 			if (DuplicatedAsset != nullptr)
 			{
+				// update duplicated asset & notify asset registry
+				if (Asset->HasAnyFlags(RF_Transient))
+				{
+					DuplicatedAsset->ClearFlags(RF_Transient);
+					DuplicatedAsset->SetFlags(RF_Public | RF_Standalone);
+				}
+
 				DuplicatedAsset->MarkPackageDirty();
 				FAssetRegistryModule::AssetCreated(DuplicatedAsset);
 				OutSavedAssets.Add(DuplicatedAsset);
+
+				// update last save directory
+				const FString PackageFilename = FPackageName::LongPackageNameToFilename(NewPackageName);
+				const FString PackagePath = FPaths::GetPath(PackageFilename);
+
+				FEditorDirectories::Get().SetLastDirectory(ELastDirectory::NEW_ASSET, PackagePath);
 			}
 			else
 			{
@@ -1008,8 +1050,7 @@ void FEditorFileUtils::SaveAssetsAs(const TArray<UObject*>& Assets, TArray<UObje
 		}
 		else
 		{
-			// save existing asset
-			OutSavedAssets.Add(Asset);
+			OutSavedAssets.Add(Asset); // save existing asset
 		}
 	}
 
diff --git a/Engine/Source/Editor/UnrealEd/Private/Kismet2/DebuggerCommands.cpp b/Engine/Source/Editor/UnrealEd/Private/Kismet2/DebuggerCommands.cpp
--- a/Engine/Source/Editor/UnrealEd/Private/Kismet2/DebuggerCommands.cpp
+++ b/Engine/Source/Editor/UnrealEd/Private/Kismet2/DebuggerCommands.cpp
@@ -29,7 +29,8 @@
 #include "Interfaces/TargetDeviceId.h"
 #include "Interfaces/ITargetPlatform.h"
 #include "Interfaces/ITargetPlatformManagerModule.h"
-#include "Interfaces/ITargetDeviceServicesModule.h"
+#include "ITargetDeviceProxy.h"
+#include "ITargetDeviceServicesModule.h"
 #include "ISettingsModule.h"
 #include "Interfaces/IMainFrameModule.h"
 
@@ -887,7 +888,7 @@ TSharedRef< SWidget > FPlayWorldCommands::GenerateLaunchMenuContent( TSharedRef<
 			if (VanillaPlatform.PlatformInfo->SDKStatus == PlatformInfo::EPlatformSDKStatus::Installed)
 			{
 				// for each platform...
-				TArray<ITargetDeviceProxyPtr> DeviceProxies;
+				TArray<TSharedPtr<ITargetDeviceProxy>> DeviceProxies;
 				TargetDeviceServicesModule->GetDeviceProxyManager()->GetProxies(VanillaPlatform.PlatformInfo->VanillaPlatformName, false, DeviceProxies);
 					
 				// if this platform had no devices, but we want to show an extra option if not installed right
@@ -904,7 +905,7 @@ TSharedRef< SWidget > FPlayWorldCommands::GenerateLaunchMenuContent( TSharedRef<
 					// for each proxy...
 					for (auto DeviceProxyIt = DeviceProxies.CreateIterator(); DeviceProxyIt; ++DeviceProxyIt)
 					{
-						ITargetDeviceProxyPtr DeviceProxy = *DeviceProxyIt;
+						TSharedPtr<ITargetDeviceProxy> DeviceProxy = *DeviceProxyIt;
 
 						// ... create an action...
 						FUIAction LaunchDeviceAction(
@@ -2270,7 +2271,7 @@ bool FInternalPlayWorldCommandCallbacks::CanLaunchOnDevice(const FString& Device
 		TSharedPtr<ITargetDeviceProxyManager> DeviceProxyManager = DeviceProxyManagerPtr.Pin();
 		if (DeviceProxyManager.IsValid())
 		{
-			ITargetDeviceProxyPtr DeviceProxy = DeviceProxyManager->FindProxy(DeviceName);
+			TSharedPtr<ITargetDeviceProxy> DeviceProxy = DeviceProxyManager->FindProxy(DeviceName);
 			return (DeviceProxy.IsValid() && DeviceProxy->IsConnected());
 		}
 	}
diff --git a/Engine/Source/Editor/UnrealEd/Private/LevelEditorViewport.cpp b/Engine/Source/Editor/UnrealEd/Private/LevelEditorViewport.cpp
--- a/Engine/Source/Editor/UnrealEd/Private/LevelEditorViewport.cpp
+++ b/Engine/Source/Editor/UnrealEd/Private/LevelEditorViewport.cpp
@@ -373,7 +373,7 @@ static bool TryAndCreateMaterialInput( UMaterial* UnrealMaterial, EMaterialKind:
 	UnrealMaterial->Expressions.Add( UnrealTextureExpression );
 	MaterialInput.Expression = UnrealTextureExpression;
 	UnrealTextureExpression->Texture = UnrealTexture;
-	UnrealTextureExpression->SamplerType = bSetupAsNormalMap ? SAMPLERTYPE_Normal : SAMPLERTYPE_Color;
+	UnrealTextureExpression->AutoSetSampleType();
 	UnrealTextureExpression->MaterialExpressionEditorX += X;
 	UnrealTextureExpression->MaterialExpressionEditorY += Y;
 
@@ -1631,6 +1631,8 @@ FSceneView* FLevelEditorViewportClient::CalcSceneView(FSceneViewFamily* ViewFami
 {
 	bWasControlledByOtherViewport = false;
 
+	UpdateViewForLockedActor();
+
 	// set all other matching viewports to my location, if the LOD locking is enabled,
 	// unless another viewport already set me this frame (otherwise they fight)
 	if (GEditor->bEnableLODLocking)
diff --git a/Engine/Source/Editor/UnrealEd/Private/PlayLevel.cpp b/Engine/Source/Editor/UnrealEd/Private/PlayLevel.cpp
--- a/Engine/Source/Editor/UnrealEd/Private/PlayLevel.cpp
+++ b/Engine/Source/Editor/UnrealEd/Private/PlayLevel.cpp
@@ -78,8 +78,8 @@
 #include "Logging/MessageLog.h"
 #include "Misc/UObjectToken.h"
 #include "Misc/MapErrors.h"
-#include "Interfaces/ITargetDeviceServicesModule.h"
-#include "Interfaces/ILauncherServicesModule.h"
+#include "ITargetDeviceServicesModule.h"
+#include "ILauncherServicesModule.h"
 #include "GameProjectGenerationModule.h"
 #include "SourceCodeNavigation.h"
 #include "PhysicsPublic.h"
@@ -3229,8 +3229,11 @@ UGameInstance* UEditorEngine::CreatePIEGameInstance(int32 InPIEInstance, bool bI
 	// By this point it is safe to remove the GameInstance from the root and allow it to garbage collected as per usual
 	GameInstance->RemoveFromRoot();
 
-	// Start the game instance
+	// Start the game instance, make sure to set the PIE instance global as this is basically a tick
+	GPlayInEditorID = InPIEInstance;
 	const FGameInstancePIEResult StartResult = GameInstance->StartPlayInEditorGameInstance(NewLocalPlayer, GameInstanceParams);
+	GPlayInEditorID = -1;
+
 	if (!StartResult.IsSuccess())
 	{
 		FMessageDialog::Open(EAppMsgType::Ok, StartResult.FailureReason);
diff --git a/Engine/Source/Editor/UnrealEd/Private/Tests/AutomationEditorCommon.cpp b/Engine/Source/Editor/UnrealEd/Private/Tests/AutomationEditorCommon.cpp
--- a/Engine/Source/Editor/UnrealEd/Private/Tests/AutomationEditorCommon.cpp
+++ b/Engine/Source/Editor/UnrealEd/Private/Tests/AutomationEditorCommon.cpp
@@ -30,17 +30,18 @@
 #include "Interfaces/IMainFrameModule.h"
 #include "ShaderCompiler.h"
 #include "AssetSelection.h"
-
-#include "Interfaces/ITargetDeviceServicesModule.h"
-#include "Interfaces/ILauncherWorker.h"
-
+#include "ITargetDeviceProxy.h"
+#include "ITargetDeviceServicesModule.h"
+#include "ILauncherWorker.h"
 #include "CookOnTheSide/CookOnTheFlyServer.h"
 #include "LightingBuildOptions.h"
 
+
 #define COOK_TIMEOUT 3600
 
 DEFINE_LOG_CATEGORY_STATIC(LogAutomationEditorCommon, Log, All);
 
+
 UWorld* FAutomationEditorCommonUtils::CreateNewMap()
 {
 	// Change out of Matinee when opening new map, so we avoid editing data in the old one.
@@ -438,12 +439,12 @@ void FAutomationEditorCommonUtils::GetLaunchOnDeviceID(FString& OutDeviceID, con
 			// shared devices section
 			ITargetDeviceServicesModule* TargetDeviceServicesModule = static_cast<ITargetDeviceServicesModule*>(FModuleManager::Get().LoadModule(TEXT("TargetDeviceServices")));
 			// for each platform...
-			TArray<ITargetDeviceProxyPtr> DeviceProxies;
+			TArray<TSharedPtr<ITargetDeviceProxy>> DeviceProxies;
 			TargetDeviceServicesModule->GetDeviceProxyManager()->GetProxies(FName(*LaunchOnSettings), true, DeviceProxies);
 			// for each proxy...
 			for (auto DeviceProxyIt = DeviceProxies.CreateIterator(); DeviceProxyIt; ++DeviceProxyIt)
 			{
-				ITargetDeviceProxyPtr DeviceProxy = *DeviceProxyIt;
+				TSharedPtr<ITargetDeviceProxy> DeviceProxy = *DeviceProxyIt;
 				if (DeviceProxy->IsConnected())
 				{
 					OutDeviceID = DeviceProxy->GetTargetDeviceId((FName)*LaunchOnSettings);
@@ -465,12 +466,12 @@ void FAutomationEditorCommonUtils::GetLaunchOnDeviceID(FString& OutDeviceID, con
 	// shared devices section
 	ITargetDeviceServicesModule* TargetDeviceServicesModule = static_cast<ITargetDeviceServicesModule*>(FModuleManager::Get().LoadModule(TEXT("TargetDeviceServices")));
 	// for each platform...
-	TArray<ITargetDeviceProxyPtr> DeviceProxies;
+	TArray<TSharedPtr<ITargetDeviceProxy>> DeviceProxies;
 	TargetDeviceServicesModule->GetDeviceProxyManager()->GetProxies(FName(*InDeviceName), true, DeviceProxies);
 	// for each proxy...
 	for (auto DeviceProxyIt = DeviceProxies.CreateIterator(); DeviceProxyIt; ++DeviceProxyIt)
 	{
-		ITargetDeviceProxyPtr DeviceProxy = *DeviceProxyIt;
+		TSharedPtr<ITargetDeviceProxy> DeviceProxy = *DeviceProxyIt;
 		if (DeviceProxy->IsConnected())
 		{
 			OutDeviceID = DeviceProxy->GetTargetDeviceId((FName)*InDeviceName);
diff --git a/Engine/Source/Editor/UnrealEd/Private/Toolkits/AssetEditorManager.cpp b/Engine/Source/Editor/UnrealEd/Private/Toolkits/AssetEditorManager.cpp
--- a/Engine/Source/Editor/UnrealEd/Private/Toolkits/AssetEditorManager.cpp
+++ b/Engine/Source/Editor/UnrealEd/Private/Toolkits/AssetEditorManager.cpp
@@ -2,13 +2,13 @@
 
 
 #include "Toolkits/AssetEditorManager.h"
-#include "Helpers/MessageEndpoint.h"
+#include "MessageEndpoint.h"
 #include "Misc/ConfigCacheIni.h"
 #include "Misc/FeedbackContext.h"
 #include "Modules/ModuleManager.h"
 #include "Settings/EditorLoadingSavingSettings.h"
 #include "AssetEditorMessages.h"
-#include "Helpers/MessageEndpointBuilder.h"
+#include "MessageEndpointBuilder.h"
 #include "IAssetTools.h"
 #include "IAssetTypeActions.h"
 #include "AssetToolsModule.h"
@@ -493,7 +493,7 @@ FArchive& operator<<(FArchive& Ar,IAssetEditorInstance*& TypeRef)
 }
 
 
-void FAssetEditorManager::HandleRequestOpenAssetMessage( const FAssetEditorRequestOpenAsset& Message, const IMessageContextRef& Context )
+void FAssetEditorManager::HandleRequestOpenAssetMessage( const FAssetEditorRequestOpenAsset& Message, const TSharedRef<IMessageContext, ESPMode::ThreadSafe>& Context )
 {
 	OpenEditorForAsset(Message.AssetName);
 }
diff --git a/Engine/Source/Programs/IOS/UnrealLaunchDaemon/Private/LaunchDaemonMessageHandler.cpp b/Engine/Source/Programs/IOS/UnrealLaunchDaemon/Private/LaunchDaemonMessageHandler.cpp
--- a/Engine/Source/Programs/IOS/UnrealLaunchDaemon/Private/LaunchDaemonMessageHandler.cpp
+++ b/Engine/Source/Programs/IOS/UnrealLaunchDaemon/Private/LaunchDaemonMessageHandler.cpp
@@ -1,9 +1,16 @@
 // Copyright 1998-2017 Epic Games, Inc. All Rights Reserved.
 
 #include "UnrealLaunchDaemonApp.h"
+
+#include "IMessageContext.h"
 #include "LaunchDaemonMessageHandler.h"
+#include "LaunchDaemonMessages.h"
+#include "MessageEndpoint.h"
+#include "MessageEndpointBuilder.h"
+
 #import <TargetConditionals.h>
 
+
 void FLaunchDaemonMessageHandler::Init()
 {
 	MessageEndpoint = FMessageEndpoint::Builder("FLaunchDaemonMessageHandler")
@@ -16,6 +23,7 @@ void FLaunchDaemonMessageHandler::Init()
 	}
 }
 
+
 void FLaunchDaemonMessageHandler::Shutdown()
 {
 	if (MessageEndpoint.IsValid())
@@ -24,7 +32,8 @@ void FLaunchDaemonMessageHandler::Shutdown()
 	}
 }
 
-void FLaunchDaemonMessageHandler::HandlePingMessage(const FIOSLaunchDaemonPing& Message, const IMessageContextRef& Context)
+
+void FLaunchDaemonMessageHandler::HandlePingMessage(const FIOSLaunchDaemonPing& Message, const TSharedRef<IMessageContext, ESPMode::ThreadSafe>& Context)
 {
 	if (MessageEndpoint.IsValid())
 	{
@@ -41,7 +50,8 @@ void FLaunchDaemonMessageHandler::HandlePingMessage(const FIOSLaunchDaemonPing&
 	}
 }
 
-void FLaunchDaemonMessageHandler::HandleLaunchRequest(const FIOSLaunchDaemonLaunchApp& Message, const IMessageContextRef& Context)
+
+void FLaunchDaemonMessageHandler::HandleLaunchRequest(const FIOSLaunchDaemonLaunchApp& Message, const TSharedRef<IMessageContext, ESPMode::ThreadSafe>& Context)
 {
 	FString Error;
 	// Leaving in extra variable in case text debugging needs to be injected.
@@ -49,6 +59,7 @@ void FLaunchDaemonMessageHandler::HandleLaunchRequest(const FIOSLaunchDaemonLaun
 	Launch(LaunchURL);
 }
 
+
 void FLaunchDaemonMessageHandler::Launch(const FString& LaunchURL)
 {
 	NSUserDefaults *defaults = [NSUserDefaults standardUserDefaults];	
diff --git a/Engine/Source/Programs/IOS/UnrealLaunchDaemon/Private/LaunchDaemonMessageHandler.h b/Engine/Source/Programs/IOS/UnrealLaunchDaemon/Private/LaunchDaemonMessageHandler.h
--- a/Engine/Source/Programs/IOS/UnrealLaunchDaemon/Private/LaunchDaemonMessageHandler.h
+++ b/Engine/Source/Programs/IOS/UnrealLaunchDaemon/Private/LaunchDaemonMessageHandler.h
@@ -2,8 +2,14 @@
 
 #pragma once
 
-#include "Messaging.h"
-#include "LaunchDaemonMessages.h"
+#include "Containers/UnrealString.h"
+#include "Templates/SharedPointer.h"
+
+class FIOSLaunchDaemonLaunchApp;
+class FIOSLaunchDaemonPing;
+class FMessageEndpoint;
+class IMessageContext;
+
 
 class FLaunchDaemonMessageHandler
 {
@@ -11,13 +17,13 @@ public:
 
 	void Init();
 	void Shutdown();
-
 	void Launch(const FString& LaunchURL);
 
 private:
-	void HandlePingMessage(const FIOSLaunchDaemonPing& Message, const IMessageContextRef& Context);
-	void HandleLaunchRequest(const FIOSLaunchDaemonLaunchApp& Message, const IMessageContextRef& Context);
 
-	FMessageEndpointPtr MessageEndpoint;
+	void HandlePingMessage(const FIOSLaunchDaemonPing& Message, const TSharedRef<IMessageContext, ESPMode::ThreadSafe>& Context);
+	void HandleLaunchRequest(const FIOSLaunchDaemonLaunchApp& Message, const TSharedRef<IMessageContext, ESPMode::ThreadSafe>& Context);
+
+	TSharedPtr<FMessageEndpoint, ESPMode::ThreadSafe> MessageEndpoint;
 	FString AppId;
 };
diff --git a/Engine/Source/Programs/UnrealBuildTool/Configuration/UEBuildTarget.cs b/Engine/Source/Programs/UnrealBuildTool/Configuration/UEBuildTarget.cs
--- a/Engine/Source/Programs/UnrealBuildTool/Configuration/UEBuildTarget.cs
+++ b/Engine/Source/Programs/UnrealBuildTool/Configuration/UEBuildTarget.cs
@@ -2381,7 +2381,10 @@ namespace UnrealBuildTool
 				{
 					foreach(UEBuildModule Module in Plugin.Modules)
 					{
-						ModuleToPlugin.Add(Module, Plugin);
+                        if (!ModuleToPlugin.ContainsKey(Module))
+                        {
+                            ModuleToPlugin.Add(Module, Plugin);
+                        }
 					}
 				}
 
@@ -3994,12 +3997,15 @@ namespace UnrealBuildTool
 						if (ModuleInfo.IsCompiledInConfiguration(Platform, TargetType, Rules.bBuildDeveloperTools, Rules.bBuildEditor, Rules.bBuildRequiresCookedData))
 						{
 							UEBuildModule Module = FindOrCreateModuleByName(ModuleInfo.Name, PluginReferenceChain);
-							if (!Module.RulesFile.IsUnderDirectory(Info.Directory))
+							if (!Instance.Modules.Contains(Module))
 							{
-								throw new BuildException("Plugin '{0}' (referenced via {1}) does not contain the '{2}' module, but lists it in '{3}'.", Info.Name, ReferenceChain, ModuleInfo.Name, Info.File);
+								if (!Module.RulesFile.IsUnderDirectory(Info.Directory))
+								{
+									throw new BuildException("Plugin '{0}' (referenced via {1}) does not contain the '{2}' module, but lists it in '{3}'.", Info.Name, ReferenceChain, ModuleInfo.Name, Info.File);
+								}
+								Instance.bDescriptorNeededAtRuntime = true;
+								Instance.Modules.Add(Module);
 							}
-							Instance.bDescriptorNeededAtRuntime = true;
-							Instance.Modules.Add(Module);
 						}
 					}
 				}
diff --git a/Engine/Source/Programs/UnrealFileServer/Private/UnrealFileServer.cpp b/Engine/Source/Programs/UnrealFileServer/Private/UnrealFileServer.cpp
--- a/Engine/Source/Programs/UnrealFileServer/Private/UnrealFileServer.cpp
+++ b/Engine/Source/Programs/UnrealFileServer/Private/UnrealFileServer.cpp
@@ -4,8 +4,8 @@
 
 #include "DirectoryWatcherModule.h"
 #include "RequiredProgramMainCPPInclude.h"
-#include "Interfaces/INetworkFileServer.h"
-#include "Interfaces/INetworkFileSystemModule.h"
+#include "INetworkFileServer.h"
+#include "INetworkFileSystemModule.h"
 #include "IDirectoryWatcher.h"
 #include "SocketSubsystem.h"
 
diff --git a/Engine/Source/Programs/UnrealFrontend/Private/Commands/LaunchFromProfileCommand.cpp b/Engine/Source/Programs/UnrealFrontend/Private/Commands/LaunchFromProfileCommand.cpp
--- a/Engine/Source/Programs/UnrealFrontend/Private/Commands/LaunchFromProfileCommand.cpp
+++ b/Engine/Source/Programs/UnrealFrontend/Private/Commands/LaunchFromProfileCommand.cpp
@@ -2,13 +2,13 @@
 
 #include "LaunchFromProfileCommand.h"
 #include "DesktopPlatformModule.h"
-#include "Interfaces/ILauncherProfile.h"
-#include "Interfaces/ILauncherProfileManager.h"
-#include "Interfaces/ILauncherTask.h"
-#include "Interfaces/ILauncherWorker.h"
-#include "Interfaces/ILauncher.h"
-#include "Interfaces/ILauncherServicesModule.h"
-#include "Interfaces/ITargetDeviceServicesModule.h"
+#include "ILauncherProfile.h"
+#include "ILauncherProfileManager.h"
+#include "ILauncherTask.h"
+#include "ILauncherWorker.h"
+#include "ILauncher.h"
+#include "ILauncherServicesModule.h"
+#include "ITargetDeviceServicesModule.h"
 #include "UserInterfaceCommand.h"
 #include "Misc/CommandLine.h"
 #include "Misc/OutputDeviceRedirector.h"
@@ -33,7 +33,7 @@ void FLaunchFromProfileCommand::Run(const FString& Params)
 
 	// Loading the Device Proxy Manager to get the needed Device Manager.
 	ITargetDeviceServicesModule& DeviceServiceModule = FModuleManager::LoadModuleChecked<ITargetDeviceServicesModule>(TEXT("TargetDeviceServices"));
-	ITargetDeviceProxyManagerRef DeviceProxyManager = DeviceServiceModule.GetDeviceProxyManager();
+	TSharedRef<ITargetDeviceProxyManager> DeviceProxyManager = DeviceServiceModule.GetDeviceProxyManager();
 
 	UE_LOG(LogUFECommands, Display, TEXT("Begin the process of launching a project using the provided profile."));
 	ILauncherRef LauncherRef = LauncherServicesModule.CreateLauncher();
diff --git a/Engine/Source/Programs/UnrealHeaderTool/Private/HeaderParser.cpp b/Engine/Source/Programs/UnrealHeaderTool/Private/HeaderParser.cpp
--- a/Engine/Source/Programs/UnrealHeaderTool/Private/HeaderParser.cpp
+++ b/Engine/Source/Programs/UnrealHeaderTool/Private/HeaderParser.cpp
@@ -965,14 +965,14 @@ namespace
 			case ECheckedMetadataSpecifier::Units:
 			{
 				// Check for numeric property
-				if (!Cast<UNumericProperty>(Field))
+				if (!Field->IsA<UNumericProperty>() && !Field->IsA<UStructProperty>())
 				{
-					FError::Throwf(TEXT("'Units' meta data can only be applied to numeric properties"));
+					FError::Throwf(TEXT("'Units' meta data can only be applied to numeric and struct properties"));
 				}
 
 				if (!FUnitConversion::UnitFromString(*InValue))
 				{
-					FError::Throwf(TEXT("Unrecognized units (%s) specified for numeric property '%s'"), *InValue, *Field->GetDisplayNameText().ToString());
+					FError::Throwf(TEXT("Unrecognized units (%s) specified for property '%s'"), *InValue, *Field->GetDisplayNameText().ToString());
 				}
 			}
 			break;
