Proper fix for UE-58864 (crash for Android derived platforms in Launch on menu)

diff --git a/Engine/Source/Developer/TargetDeviceServices/Private/Proxies/TargetDeviceProxy.cpp b/Engine/Source/Developer/TargetDeviceServices/Private/Proxies/TargetDeviceProxy.cpp
--- a/Engine/Source/Developer/TargetDeviceServices/Private/Proxies/TargetDeviceProxy.cpp
+++ b/Engine/Source/Developer/TargetDeviceServices/Private/Proxies/TargetDeviceProxy.cpp
@@ -67,8 +67,7 @@ void FTargetDeviceProxy::UpdateFromMessage( const FTargetDeviceServicePong& Mess
 	SupportsPowerOn = Message.SupportsPowerOn;
 	SupportsReboot = Message.SupportsReboot;
 	SupportsVariants = Message.SupportsVariants;
-	//Temp disabling the all devices default variant due to Lumin devices deriving off of Android.  When Lumin is plugged in alone, there was a crash in GetTargetDeviceIds.
-	DefaultVariant = Message.DefaultVariant;// Aggregated ? Message.AllDevicesDefaultVariant : Message.DefaultVariant;
+	DefaultVariant = Message.DefaultVariant;
 
 	// Update the map of flavors.
 	for (int Index = 0; Index < Message.Variants.Num(); Index++)
@@ -83,6 +82,12 @@ void FTargetDeviceProxy::UpdateFromMessage( const FTargetDeviceServicePong& Mess
 		Variant.TargetPlatformId = MsgVariant.TargetPlatformId;
 		Variant.VanillaPlatformId = MsgVariant.VanillaPlatformId;
 		Variant.PlatformDisplayName = FText::FromString(MsgVariant.PlatformDisplayName);
+
+		if (Aggregated && MsgVariant.TargetPlatformId.IsEqual(Message.AllDevicesDefaultVariant))
+		{
+			// for aggregated platforms,check if the declared AllDevicesDefaultVariant is supported by at least one device
+			DefaultVariant = Message.AllDevicesDefaultVariant;
+		}
 	}
 
 	LastUpdateTime = FDateTime::UtcNow();
