Fix allocator inconsistencies in TargetDeviceService.

- Flagged by ASan on Linux.

diff --git a/Engine/Source/Developer/TargetDeviceServices/Private/Services/TargetDeviceService.cpp b/Engine/Source/Developer/TargetDeviceServices/Private/Services/TargetDeviceService.cpp
--- a/Engine/Source/Developer/TargetDeviceServices/Private/Services/TargetDeviceService.cpp
+++ b/Engine/Source/Developer/TargetDeviceServices/Private/Services/TargetDeviceService.cpp
@@ -245,7 +245,9 @@ bool FTargetDeviceService::Start()
 		ClaimHost = FPlatformProcess::ComputerName();
 		ClaimUser = FPlatformProcess::UserName(false);
 
-		MessageEndpoint->Publish(new FTargetDeviceClaimed(DeviceName, ClaimHost, ClaimUser));
+		// message is going to be deleted by FMemory::Free() (see FMessageContext destructor), so allocate it with Malloc
+		void* Memory = FMemory::Malloc(sizeof(FTargetDeviceClaimed));
+		MessageEndpoint->Publish(new(Memory) FTargetDeviceClaimed(DeviceName, ClaimHost, ClaimUser));
 
 		Running = true;
 	}
@@ -258,8 +260,9 @@ void FTargetDeviceService::Stop()
 {
 	if (Running)
 	{
-	
-		MessageEndpoint->Publish(new FTargetDeviceUnclaimed(DeviceName, FPlatformProcess::ComputerName(), FPlatformProcess::UserName(false)));
+		// message is going to be deleted by FMemory::Free() (see FMessageContext destructor), so allocate it with Malloc
+		void* Memory = FMemory::Malloc(sizeof(FTargetDeviceUnclaimed));
+		MessageEndpoint->Publish(new(Memory) FTargetDeviceUnclaimed(DeviceName, FPlatformProcess::ComputerName(), FPlatformProcess::UserName(false)));
 		FPlatformProcess::SleepNoStats(0.01);
 
 		// Only stop the device if we care about device claiming
@@ -356,7 +359,9 @@ void FTargetDeviceService::HandleClaimedMessage(const FTargetDeviceClaimed& Mess
 	{
 		if (Context->GetSender() != MessageEndpoint->GetAddress())
 		{
-			MessageEndpoint->Send(new FTargetDeviceClaimDenied(DeviceName, FPlatformProcess::ComputerName(), FPlatformProcess::UserName(false)), Context->GetSender());
+			// message is going to be deleted by FMemory::Free() (see FMessageContext destructor), so allocate it with Malloc
+			void* Memory = FMemory::Malloc(sizeof(FTargetDeviceClaimDenied));
+			MessageEndpoint->Send(new(Memory) FTargetDeviceClaimDenied(DeviceName, FPlatformProcess::ComputerName(), FPlatformProcess::UserName(false)), Context->GetSender());
 		}
 	}
 	else
@@ -425,7 +430,9 @@ void FTargetDeviceService::HandleDeployCommitMessage(const FTargetDeviceServiceD
 		bool Succeeded = TargetDevice->Deploy(SourceFolder, OutAppId);
 
 		IFileManager::Get().DeleteDirectory(*SourceFolder, false, true);
-		MessageEndpoint->Send(new FTargetDeviceServiceDeployFinished(Message.Variant, OutAppId, Succeeded, Message.TransactionId), Context->GetSender());
+		// message is going to be deleted by FMemory::Free() (see FMessageContext destructor), so allocate it with Malloc
+		void* Memory = FMemory::Malloc(sizeof(FTargetDeviceServiceDeployFinished));
+		MessageEndpoint->Send(new(Memory) FTargetDeviceServiceDeployFinished(Message.Variant, OutAppId, Succeeded, Message.TransactionId), Context->GetSender());
 	}
 }
 
@@ -460,7 +467,9 @@ void FTargetDeviceService::HandleLaunchAppMessage(const FTargetDeviceServiceLaun
 
 		if (MessageEndpoint.IsValid())
 		{
-			MessageEndpoint->Send(new FTargetDeviceServiceLaunchFinished(Message.AppID, ProcessId, Succeeded), Context->GetSender());
+			// message is going to be deleted by FMemory::Free() (see FMessageContext destructor), so allocate it with Malloc
+			void* Memory = FMemory::Malloc(sizeof(FTargetDeviceServiceLaunchFinished));
+			MessageEndpoint->Send(new(Memory) FTargetDeviceServiceLaunchFinished(Message.AppID, ProcessId, Succeeded), Context->GetSender());
 		}
 	}
 }
@@ -485,7 +494,9 @@ void FTargetDeviceService::HandlePingMessage(const FTargetDeviceServicePing& InM
 		const FString& PlatformName = DefaultDevice->GetTargetPlatform().PlatformName();
 		const PlatformInfo::FPlatformInfo* VanillaInfo = PlatformInfo::FindVanillaPlatformInfo(FName(*PlatformName));
 
-		FTargetDeviceServicePong* Message = new FTargetDeviceServicePong();
+		// message is going to be deleted by FMemory::Free() (see FMessageContext destructor), so allocate it with Malloc
+		void* Memory = FMemory::Malloc(sizeof(FTargetDeviceServicePong));
+		FTargetDeviceServicePong* Message = new(Memory) FTargetDeviceServicePong();
 
 		Message->Name = DefaultDevice->GetName();
 		Message->Type = TargetDeviceTypes::ToString(DefaultDevice->GetDeviceType());
@@ -595,6 +606,8 @@ void FTargetDeviceService::HandleRunExecutableMessage(const FTargetDeviceService
 		uint32 OutProcessId;
 		bool Succeeded = TargetDevice->Run(Message.ExecutablePath, Message.Params, &OutProcessId);
 
-		MessageEndpoint->Send(new FTargetDeviceServiceRunFinished(Message.Variant, Message.ExecutablePath, OutProcessId, Succeeded), Context->GetSender());
+		// message is going to be deleted by FMemory::Free() (see FMessageContext destructor), so allocate it with Malloc
+		void* Memory = FMemory::Malloc(sizeof(FTargetDeviceServiceRunFinished));
+		MessageEndpoint->Send(new(Memory) FTargetDeviceServiceRunFinished(Message.Variant, Message.ExecutablePath, OutProcessId, Succeeded), Context->GetSender());
 	}
 }
