Fix allocator inconsistencies in TargetDeviceProxyService.

- Flagged by ASan on Linux.

diff --git a/Engine/Source/Developer/TargetDeviceServices/Private/Proxies/TargetDeviceProxyManager.cpp b/Engine/Source/Developer/TargetDeviceServices/Private/Proxies/TargetDeviceProxyManager.cpp
--- a/Engine/Source/Developer/TargetDeviceServices/Private/Proxies/TargetDeviceProxyManager.cpp
+++ b/Engine/Source/Developer/TargetDeviceServices/Private/Proxies/TargetDeviceProxyManager.cpp
@@ -141,6 +141,8 @@ void FTargetDeviceProxyManager::SendPing()
 
 	if (MessageEndpoint.IsValid())
 	{
+		// message is going to be deleted by FMemory::Free() (see FMessageContext destructor), so allocate it with Malloc
+		void* Memory = FMemory::Malloc(sizeof(FTargetDeviceServicePing));
 		MessageEndpoint->Publish(new FTargetDeviceServicePing(FPlatformProcess::UserName(false)), EMessageScope::Network);
 	}
 }
