Fixed intermittent launch failures when cook-on-the-fly is used.

A timeout deletes target devices if they fail to respond in 7.5s. The timeout logic was flawed and triggered when the device manager wasn't pumped for 7.5s - even though ping responses came in very quickly. The new logic clears the timeout on a ping response.

Note that there is still a >10s long editor freeze during which the network filesystem seems to clear its cache. There is LogNetworkPlatformFile log spam like below during the freeze:
  Deleting cached file: ... Module.RawInput.gen.cpp
  Deleting cached file: ... Module.RawInput.gen.cpp.obj
  Deleting cached file: ... Module.RawInput.gen.cpp.old
  Deleting cached file: ... Module.RawInput.gen.cpp.txt
  Deleting cached file: ... UE4Editor-RawInput.exp
  Deleting cached file: ... UE4Editor-RawInput.lib

Bonus drive-by changes: Fixed a  memory leak and aligned up some allocations.

(ushell-p4-cherrypick of 7840548 by [REDACTED_PERSON])

diff --git a/Engine/Source/Developer/TargetDeviceServices/Private/Proxies/TargetDeviceProxy.cpp b/Engine/Source/Developer/TargetDeviceServices/Private/Proxies/TargetDeviceProxy.cpp
--- a/Engine/Source/Developer/TargetDeviceServices/Private/Proxies/TargetDeviceProxy.cpp
+++ b/Engine/Source/Developer/TargetDeviceServices/Private/Proxies/TargetDeviceProxy.cpp
@@ -15,7 +15,6 @@
 FTargetDeviceProxy::FTargetDeviceProxy(const FString& InName)
 	: Connected(false)
 	, Authorized(false)
-	, LastUpdateTime(0)
 	, Name(InName)
 	, SupportsMultiLaunch(false)
 	, SupportsPowerOff(false)
@@ -93,7 +92,7 @@ void FTargetDeviceProxy::UpdateFromMessage( const FTargetDeviceServicePong& Mess
 		}
 	}
 
-	LastUpdateTime = FDateTime::UtcNow();
+	PingTimeout.Reset();
 }
 
 
diff --git a/Engine/Source/Developer/TargetDeviceServices/Private/Proxies/TargetDeviceProxy.h b/Engine/Source/Developer/TargetDeviceServices/Private/Proxies/TargetDeviceProxy.h
--- a/Engine/Source/Developer/TargetDeviceServices/Private/Proxies/TargetDeviceProxy.h
+++ b/Engine/Source/Developer/TargetDeviceServices/Private/Proxies/TargetDeviceProxy.h
@@ -5,8 +5,8 @@
 #include "CoreTypes.h"
 #include "IMessageContext.h"
 #include "MessageEndpoint.h"
-
 #include "ITargetDeviceProxy.h"
+#include "Misc/Optional.h"
 
 class FTargetDeviceProxy;
 class IMessageContext;
@@ -51,17 +51,26 @@ public:
 public:
 
 	/**
-	 * Gets the time at which the proxy was last updated.
-	 *
-	 * @return Date and time of the last update.
+	 * Add a time by which a ping response should arrive
+	 */
+	void AddTimeout(const FDateTime& NewTimeout)
+	{
+		if (!PingTimeout || NewTimeout < PingTimeout.GetValue())
+		{
+			PingTimeout = NewTimeout;
+		}
+	}
+
+	/**
+	 * Check if ping response arrived on time
 	 */
-	const FDateTime& GetLastUpdateTime() const
+	bool HasTimedOut(const FDateTime& Now) const
 	{
-		return LastUpdateTime;
+		return PingTimeout && Now > PingTimeout.GetValue();
 	}
 
 	/**
-	 * Updates the proxy's information from the given device service response.
+	 * Updates the proxy's information from the given device service response and clear timeout expectations
 	 *
 	 * @param Message The message containing the response.
 	 * @param Context The message context.
@@ -224,8 +233,8 @@ private:
 	/** Holds the name of the user that owns the device. */
 	FString HostUser;
 
-	/** Holds the time at which the last ping reply was received. */
-	FDateTime LastUpdateTime;
+	/** The time by which the lack of any ping response is considered a timeout. */
+	TOptional<FDateTime> PingTimeout;
 
 	/** Holds the device make. */
 	FString Make;
diff --git a/Engine/Source/Developer/TargetDeviceServices/Private/Proxies/TargetDeviceProxyManager.cpp b/Engine/Source/Developer/TargetDeviceServices/Private/Proxies/TargetDeviceProxyManager.cpp
--- a/Engine/Source/Developer/TargetDeviceServices/Private/Proxies/TargetDeviceProxyManager.cpp
+++ b/Engine/Source/Developer/TargetDeviceServices/Private/Proxies/TargetDeviceProxyManager.cpp
@@ -125,7 +125,7 @@ void FTargetDeviceProxyManager::RemoveDeadProxies()
 
 	for (auto ProxyIter = Proxies.CreateIterator(); ProxyIter; ++ProxyIter)
 	{
-		if (ProxyIter.Value()->GetLastUpdateTime() + FTimespan::FromSeconds(3.0 * TARGET_DEVICE_SERVICES_PING_INTERVAL) < CurrentTime)
+		if (ProxyIter.Value()->HasTimedOut(CurrentTime))
 		{
 			TSharedPtr<ITargetDeviceProxy> RemovedProxy = ProxyIter.Value();
 			ProxyIter.RemoveCurrent();
@@ -141,9 +141,16 @@ void FTargetDeviceProxyManager::SendPing()
 
 	if (MessageEndpoint.IsValid())
 	{
+		// Add a timeout to all known devices
+		const FDateTime Timeout = FDateTime::UtcNow() + FTimespan:: FromSeconds(3 * TARGET_DEVICE_SERVICES_PING_INTERVAL);
+		for (const TPair<FString, TSharedPtr<FTargetDeviceProxy>>& Proxy : Proxies)
+		{
+			Proxy.Value->AddTimeout(Timeout);
+		}
+
 		// message is going to be deleted by FMemory::Free() (see FMessageContext destructor), so allocate it with Malloc
-		void* Memory = FMemory::Malloc(sizeof(FTargetDeviceServicePing));
-		MessageEndpoint->Publish(new FTargetDeviceServicePing(FPlatformProcess::UserName(false)), EMessageScope::Network);
+		void* Memory = FMemory::Malloc(sizeof(FTargetDeviceServicePing), alignof(FTargetDeviceServicePing));
+		MessageEndpoint->Publish(new (Memory) FTargetDeviceServicePing(FPlatformProcess::UserName(false)), EMessageScope::Network);
 	}
 }
 
diff --git a/Engine/Source/Developer/TargetDeviceServices/Private/Services/TargetDeviceService.cpp b/Engine/Source/Developer/TargetDeviceServices/Private/Services/TargetDeviceService.cpp
--- a/Engine/Source/Developer/TargetDeviceServices/Private/Services/TargetDeviceService.cpp
+++ b/Engine/Source/Developer/TargetDeviceServices/Private/Services/TargetDeviceService.cpp
@@ -246,7 +246,7 @@ bool FTargetDeviceService::Start()
 		ClaimUser = FPlatformProcess::UserName(false);
 
 		// message is going to be deleted by FMemory::Free() (see FMessageContext destructor), so allocate it with Malloc
-		void* Memory = FMemory::Malloc(sizeof(FTargetDeviceClaimed));
+		void* Memory = FMemory::Malloc(sizeof(FTargetDeviceClaimed), alignof(FTargetDeviceClaimed));
 		MessageEndpoint->Publish(new(Memory) FTargetDeviceClaimed(DeviceName, ClaimHost, ClaimUser));
 
 		Running = true;
@@ -261,7 +261,7 @@ void FTargetDeviceService::Stop()
 	if (Running)
 	{
 		// message is going to be deleted by FMemory::Free() (see FMessageContext destructor), so allocate it with Malloc
-		void* Memory = FMemory::Malloc(sizeof(FTargetDeviceUnclaimed));
+		void* Memory = FMemory::Malloc(sizeof(FTargetDeviceUnclaimed), alignof(FTargetDeviceUnclaimed));
 		MessageEndpoint->Publish(new(Memory) FTargetDeviceUnclaimed(DeviceName, FPlatformProcess::ComputerName(), FPlatformProcess::UserName(false)));
 		FPlatformProcess::SleepNoStats(0.01);
 
@@ -360,7 +360,7 @@ void FTargetDeviceService::HandleClaimedMessage(const FTargetDeviceClaimed& Mess
 		if (Context->GetSender() != MessageEndpoint->GetAddress())
 		{
 			// message is going to be deleted by FMemory::Free() (see FMessageContext destructor), so allocate it with Malloc
-			void* Memory = FMemory::Malloc(sizeof(FTargetDeviceClaimDenied));
+			void* Memory = FMemory::Malloc(sizeof(FTargetDeviceClaimDenied), alignof(FTargetDeviceClaimDenied));
 			MessageEndpoint->Send(new(Memory) FTargetDeviceClaimDenied(DeviceName, FPlatformProcess::ComputerName(), FPlatformProcess::UserName(false)), Context->GetSender());
 		}
 	}
@@ -431,7 +431,7 @@ void FTargetDeviceService::HandleDeployCommitMessage(const FTargetDeviceServiceD
 
 		IFileManager::Get().DeleteDirectory(*SourceFolder, false, true);
 		// message is going to be deleted by FMemory::Free() (see FMessageContext destructor), so allocate it with Malloc
-		void* Memory = FMemory::Malloc(sizeof(FTargetDeviceServiceDeployFinished));
+		void* Memory = FMemory::Malloc(sizeof(FTargetDeviceServiceDeployFinished), alignof(FTargetDeviceServiceDeployFinished));
 		MessageEndpoint->Send(new(Memory) FTargetDeviceServiceDeployFinished(Message.Variant, OutAppId, Succeeded, Message.TransactionId), Context->GetSender());
 	}
 }
@@ -468,7 +468,7 @@ void FTargetDeviceService::HandleLaunchAppMessage(const FTargetDeviceServiceLaun
 		if (MessageEndpoint.IsValid())
 		{
 			// message is going to be deleted by FMemory::Free() (see FMessageContext destructor), so allocate it with Malloc
-			void* Memory = FMemory::Malloc(sizeof(FTargetDeviceServiceLaunchFinished));
+			void* Memory = FMemory::Malloc(sizeof(FTargetDeviceServiceLaunchFinished), alignof(FTargetDeviceServiceLaunchFinished));
 			MessageEndpoint->Send(new(Memory) FTargetDeviceServiceLaunchFinished(Message.AppID, ProcessId, Succeeded), Context->GetSender());
 		}
 	}
@@ -495,7 +495,7 @@ void FTargetDeviceService::HandlePingMessage(const FTargetDeviceServicePing& InM
 		const PlatformInfo::FPlatformInfo* VanillaInfo = PlatformInfo::FindVanillaPlatformInfo(FName(*PlatformName));
 
 		// message is going to be deleted by FMemory::Free() (see FMessageContext destructor), so allocate it with Malloc
-		void* Memory = FMemory::Malloc(sizeof(FTargetDeviceServicePong));
+		void* Memory = FMemory::Malloc(sizeof(FTargetDeviceServicePong), alignof(FTargetDeviceServicePong));
 		FTargetDeviceServicePong* Message = new(Memory) FTargetDeviceServicePong();
 
 		Message->Name = DefaultDevice->GetName();
@@ -607,7 +607,7 @@ void FTargetDeviceService::HandleRunExecutableMessage(const FTargetDeviceService
 		bool Succeeded = TargetDevice->Run(Message.ExecutablePath, Message.Params, &OutProcessId);
 
 		// message is going to be deleted by FMemory::Free() (see FMessageContext destructor), so allocate it with Malloc
-		void* Memory = FMemory::Malloc(sizeof(FTargetDeviceServiceRunFinished));
+		void* Memory = FMemory::Malloc(sizeof(FTargetDeviceServiceRunFinished), alignof(FTargetDeviceServiceRunFinished));
 		MessageEndpoint->Send(new(Memory) FTargetDeviceServiceRunFinished(Message.Variant, Message.ExecutablePath, OutProcessId, Succeeded), Context->GetSender());
 	}
 }
