- Added a log to help QA track down an oddity I can't repro

diff --git a/Engine/Source/Developer/TargetDeviceServices/Private/Services/TargetDeviceService.cpp b/Engine/Source/Developer/TargetDeviceServices/Private/Services/TargetDeviceService.cpp
--- a/Engine/Source/Developer/TargetDeviceServices/Private/Services/TargetDeviceService.cpp
+++ b/Engine/Source/Developer/TargetDeviceServices/Private/Services/TargetDeviceService.cpp
@@ -83,6 +83,8 @@ void FTargetDeviceService::AddTargetDevice(TSharedPtr<ITargetDevice, ESPMode::Th
 
 	FName Variant = FName(InDevice->GetTargetPlatform().PlatformName().GetCharArray().GetData());
 
+	UE_LOG(LogTemp, Display, TEXT("TPDEBUG: Adding target device %s to target services"), *Variant.ToString());
+
 	if (DevicePlatformName == NAME_None)
 	{
 		// If this seems nasty your right!
