- Disabled the editor device sharing between machines/copies of the editor, but only broadcasting the Ping message to the current Process (instead of Network). Doesn't remove any code complexity, but is a safe solution

diff --git a/Engine/Source/Developer/TargetDeviceServices/Private/Proxies/TargetDeviceProxyManager.cpp b/Engine/Source/Developer/TargetDeviceServices/Private/Proxies/TargetDeviceProxyManager.cpp
--- a/Engine/Source/Developer/TargetDeviceServices/Private/Proxies/TargetDeviceProxyManager.cpp
+++ b/Engine/Source/Developer/TargetDeviceServices/Private/Proxies/TargetDeviceProxyManager.cpp
@@ -150,7 +150,8 @@ void FTargetDeviceProxyManager::SendPing()
 
 		// message is going to be deleted by FMemory::Free() (see FMessageContext destructor), so allocate it with Malloc
 		void* Memory = FMemory::Malloc(sizeof(FTargetDeviceServicePing), alignof(FTargetDeviceServicePing));
-		MessageEndpoint->Publish(new (Memory) FTargetDeviceServicePing(FPlatformProcess::UserName(false)), EMessageScope::Network);
+		// send to only this process, as we no longer want to send out pings to the local network and get another machine's list of devices
+		MessageEndpoint->Publish(new (Memory) FTargetDeviceServicePing(FPlatformProcess::UserName(false)), EMessageScope::Process);
 	}
 }
 
