- Removing unused ITargetDevice functionality (like Deploy and Run in C++). These are are supported in some platforms but not all, and are only used by undocumented UFE "commands". UAT is expected to handle all of this.

diff --git a/Engine/Plugins/Runtime/AR/Microsoft/HoloLensAR/Source/HoloLensTargetPlatform/Private/HoloLensTargetDevice.cpp b/Engine/Plugins/Runtime/AR/Microsoft/HoloLensAR/Source/HoloLensTargetPlatform/Private/HoloLensTargetDevice.cpp
--- a/Engine/Plugins/Runtime/AR/Microsoft/HoloLensAR/Source/HoloLensTargetPlatform/Private/HoloLensTargetDevice.cpp
+++ b/Engine/Plugins/Runtime/AR/Microsoft/HoloLensAR/Source/HoloLensTargetPlatform/Private/HoloLensTargetDevice.cpp
@@ -410,11 +410,6 @@ bool FHoloLensTargetDevice::SupportsFeature(ETargetDeviceFeatures Feature) const
 	}
 }
 
-bool FHoloLensTargetDevice::SupportsSdkVersion(const FString& VersionString) const 
-{
-	return false;
-}
-
 bool FHoloLensTargetDevice::TerminateProcess(const int64 ProcessId) 
 {
 	if (Info.IsLocal)
@@ -445,45 +440,6 @@ bool FHoloLensTargetDevice::TerminateLaunchedProcess(const FString & ProcessIden
 	return true;
 }
 
-
-bool FHoloLensTargetDevice::Deploy(const FString& SourceFolder, FString& OutAppId)
-{
-	return true;
-}
-
-bool FHoloLensTargetDevice::Launch(const FString& AppId, EBuildConfiguration BuildConfiguration, EBuildTargetType TargetType, const FString& Params, uint32* OutProcessId)
-{
-	return false;
-}
-
-bool FHoloLensTargetDevice::Run(const FString& ExecutablePath, const FString& Params, uint32* OutProcessId)
-{
-	HRESULT hr = CoInitialize(nullptr);
-	if (FAILED(hr))
-	{
-		UE_LOG(LogTemp, Warning, TEXT("FHoloLensTargetDevice::Run - CoInitialize() failed with hr = 0x(%x)"), hr);
-	}
-
-	// Currently even packaged builds get an exe name in here which kind of works because we 
-	// don't yet support remote deployment and so the loose structure the package was created 
-	// from is probably in place on this machine.  So the code will read the manifest from the
-	// loose version, but actually launch the package (since that's what's registered).
-	auto Extension = FPaths::GetExtension(ExecutablePath);
-	bool PathIsActuallyPackage = Extension.StartsWith(TEXT("appx")) || Extension.StartsWith(TEXT("msix"));
-	FString StreamPath;
-	if (PathIsActuallyPackage)
-	{
-		StreamPath = ExecutablePath;
-	}
-	else
-	{
-		StreamPath = FPaths::Combine(*FPaths::GetPath(ExecutablePath), TEXT("../../.."));
-		StreamPath = FPaths::Combine(*StreamPath, TEXT("AppxManifest.xml"));
-	}
-
-	return WindowsMixedReality::PackageProject(*StreamPath, PathIsActuallyPackage, *Params, OutProcessId);
-}
-
 void FHoloLensTargetDevice::StartHeartBeat()
 {
 	FScopeLock lock(&CriticalSection);
diff --git a/Engine/Plugins/Runtime/AR/Microsoft/HoloLensAR/Source/HoloLensTargetPlatform/Private/HoloLensTargetDevice.h b/Engine/Plugins/Runtime/AR/Microsoft/HoloLensAR/Source/HoloLensTargetPlatform/Private/HoloLensTargetDevice.h
--- a/Engine/Plugins/Runtime/AR/Microsoft/HoloLensAR/Source/HoloLensTargetPlatform/Private/HoloLensTargetDevice.h
+++ b/Engine/Plugins/Runtime/AR/Microsoft/HoloLensAR/Source/HoloLensTargetPlatform/Private/HoloLensTargetDevice.h
@@ -38,8 +38,6 @@ public:
 
 	virtual bool Connect() override;
 
-	virtual bool Deploy(const FString& SourceFolder, FString& OutAppId) override;
-
 	virtual void Disconnect() override;
 
 	virtual ETargetDeviceTypes GetDeviceType() const override;
@@ -62,22 +60,16 @@ public:
 
 	virtual bool IsDefault() const override;
 
-	virtual bool Launch(const FString& AppId, EBuildConfiguration BuildConfiguration, EBuildTargetType TargetType, const FString& Params, uint32* OutProcessId);
-
 	virtual bool PowerOff(bool Force) override;
 
 	virtual bool PowerOn() override;
 
 	virtual bool Reboot(bool bReconnect = false) override;
 
-	virtual bool Run(const FString& ExecutablePath, const FString& Params, uint32* OutProcessId) override;
-
 	virtual void SetUserCredentials(const FString& UserName, const FString& UserPassword) override;
 
 	virtual bool SupportsFeature(ETargetDeviceFeatures Feature) const override;
 
-	virtual bool SupportsSdkVersion(const FString& VersionString) const override;
-
 	virtual bool TerminateProcess(const int64 ProcessId) override;
 
 	virtual bool TerminateLaunchedProcess(const FString & ProcessIdentifier) override;
diff --git a/Engine/Source/Developer/Android/AndroidTargetPlatform/Private/AndroidTargetDevice.h b/Engine/Source/Developer/Android/AndroidTargetPlatform/Private/AndroidTargetDevice.h
--- a/Engine/Source/Developer/Android/AndroidTargetPlatform/Private/AndroidTargetDevice.h
+++ b/Engine/Source/Developer/Android/AndroidTargetPlatform/Private/AndroidTargetDevice.h
@@ -127,8 +127,6 @@ public:
 		return true;
 	}
 
-	virtual bool Deploy(const FString& SourceFolder, FString& OutAppId) override;
-
 	virtual void Disconnect() override
 	{
 	}
@@ -198,12 +196,9 @@ public:
 		return "Android_Multi";
 	}
 
-	virtual bool Launch(const FString& AppId, EBuildConfiguration BuildConfiguration, EBuildTargetType TargetType, const FString& Params, uint32* OutProcessId);
 	virtual bool Reboot(bool bReconnect = false) override;
-	virtual bool Run(const FString& ExecutablePath, const FString& Params, uint32* OutProcessId) override;
 	virtual bool TerminateLaunchedProcess(const FString& ProcessIdentifier) override;
 	virtual bool SupportsFeature(ETargetDeviceFeatures Feature) const override;
-	virtual bool SupportsSdkVersion(const FString& VersionString) const override;
 	virtual bool TerminateProcess(const int64 ProcessId) override;
 	virtual void SetUserCredentials(const FString& UserName, const FString& UserPassword) override;
 	virtual bool GetUserCredentials(FString& OutUserName, FString& OutUserPassword) override;
diff --git a/Engine/Source/Developer/Android/AndroidTargetPlatform/Private/AndroidTargetDevice.inl b/Engine/Source/Developer/Android/AndroidTargetPlatform/Private/AndroidTargetDevice.inl
--- a/Engine/Source/Developer/Android/AndroidTargetPlatform/Private/AndroidTargetDevice.inl
+++ b/Engine/Source/Developer/Android/AndroidTargetPlatform/Private/AndroidTargetDevice.inl
@@ -14,22 +14,6 @@ enum class ETargetDeviceFeatures;
 
 template<typename OptionalType> struct TOptional;
 
-inline bool FAndroidTargetDevice::Deploy( const FString& SourceFolder, FString& OutAppId )
-{
-	int32 ReturnCode = 0;
-	//void* RunningProc = NULL;
-
-	// run the packager to create and install the .apk
-	// @todo android: install separately?
-	FString RepackageCommand(FString(TEXT("AndroidPackager ")) + OutAppId + FString(TEXT(" AndroidARMv7 ")) + FString(TEXT("Debug")));
-	FProcHandle RunningProc = FPlatformProcess::CreateProc(TEXT("../DotNET/Android/AndroidPackager"), *RepackageCommand, true, false, false, NULL, 0, TEXT("../DotNET/"), NULL);
-	FPlatformProcess::WaitForProc(RunningProc);
-	FPlatformProcess::GetProcReturnCode(RunningProc, &ReturnCode);
-	FPlatformProcess::CloseProc(RunningProc);
-
-	return 0 == ReturnCode;
-}
-
 inline FString FAndroidTargetDevice::GetOperatingSystemName()
 {
 	if (!AndroidVersionString.IsEmpty())
@@ -48,13 +32,6 @@ inline int32 FAndroidTargetDevice::GetProcessSnapshot( TArray<FTargetDeviceProce
 }
 
 
-inline bool FAndroidTargetDevice::Launch( const FString& AppId, EBuildConfiguration BuildConfiguration, EBuildTargetType TargetType, const FString& Params, uint32* OutProcessId )
-{
-	// this isn't used, UAT handles it all
-	return false;
-}
-
-
 inline bool FAndroidTargetDevice::Reboot( bool bReconnect )
 {
 	if (!ExecuteAdbCommand(TEXT("reboot"), NULL, NULL))
@@ -80,12 +57,6 @@ inline FString FAndroidTargetDevice::GetAllDevicesName() const
 	return FString::Printf(TEXT("All_%s_On_%s"), *(GetTargetPlatform().IniPlatformName()), FPlatformProcess::ComputerName());
 }
 
-inline bool FAndroidTargetDevice::Run( const FString& ExecutablePath, const FString& Params, uint32* OutProcessId )
-{
-	// @todo android: how to run from this?
-	return false;
-}
-
 // cancel the running application
 inline bool FAndroidTargetDevice::TerminateLaunchedProcess(const FString& ProcessIdentifier)
 {
@@ -112,12 +83,6 @@ inline bool FAndroidTargetDevice::SupportsFeature( ETargetDeviceFeatures Feature
 }
 
 
-inline bool FAndroidTargetDevice::SupportsSdkVersion( const FString& VersionString ) const
-{
-	return true;
-}
-
-
 inline bool FAndroidTargetDevice::TerminateProcess( const int64 ProcessId )
 {
 	return false;
diff --git a/Engine/Source/Developer/IOS/IOSTargetPlatform/Private/IOSTargetDevice.cpp b/Engine/Source/Developer/IOS/IOSTargetPlatform/Private/IOSTargetDevice.cpp
--- a/Engine/Source/Developer/IOS/IOSTargetPlatform/Private/IOSTargetDevice.cpp
+++ b/Engine/Source/Developer/IOS/IOSTargetPlatform/Private/IOSTargetDevice.cpp
@@ -304,11 +304,6 @@ bool FIOSTargetDevice::Connect()
 	return true;
 }
 
-bool FIOSTargetDevice::Deploy(const FString& SourceFolder, FString& OutAppId)
-{
-	return false;
-}
-
 void FIOSTargetDevice::Disconnect()
 {
 }
@@ -352,17 +347,6 @@ bool FIOSTargetDevice::IsDefault() const
 {
 	return true;
 }
-	
-bool FIOSTargetDevice::Launch(const FString& InAppId, EBuildConfiguration InBuildConfiguration, EBuildTargetType TargetType, const FString& Params, uint32* OutProcessId)
-{
-#if !PLATFORM_MAC
-	MessageEndpoint->Send(new FIOSLaunchDaemonLaunchApp(InAppId, Params), DeviceEndpoint);
-	return true;
-#else
-	//Set return to false on Mac, since we could not find a way to do remote deploy/launch.
-	return false;
-#endif // !PLATFORM_MAC
-}
 
 bool FIOSTargetDevice::PowerOff(bool Force)
 {
@@ -385,43 +369,6 @@ bool FIOSTargetDevice::Reboot(bool bReconnect)
 	return false;
 }
 
-bool FIOSTargetDevice::Run(const FString& ExecutablePath, const FString& Params, uint32* OutProcessId)
-{
-#if !PLATFORM_MAC
-	// The executable path usually looks something like this: directory/<gamename>.stub
-	// We just need <gamename>, so strip that out.
-	int32 LastPeriodPos = ExecutablePath.Find( TEXT("."), ESearchCase::IgnoreCase, ESearchDir::FromEnd);
-	int32 LastBackslashPos = ExecutablePath.Find( TEXT("\\"), ESearchCase::IgnoreCase, ESearchDir::FromEnd);
-	int32 LastSlashPos = ExecutablePath.Find( TEXT("/"), ESearchCase::IgnoreCase, ESearchDir::FromEnd);
-	int32 TrimPos = LastBackslashPos > LastSlashPos ? LastBackslashPos : LastSlashPos;
-	if ( TrimPos > LastPeriodPos )
-	{
-		// Ignore any periods in the path before the last "/" or "\"
-		LastPeriodPos = ExecutablePath.Len() - 1;
-	}
-
-	if ( TrimPos == INDEX_NONE )
-	{
-		TrimPos = 0; // take the whole string from the beginning
-	}
-	else
-	{
-		// increment to one character beyond the slash
-		TrimPos++;
-	}
-
-	int32 Count = LastPeriodPos - TrimPos;
-	FString NewAppId = ExecutablePath.Mid(TrimPos, Count); // remove the ".stub" and the proceeding directory to get the game name
-
-	SetAppId(NewAppId);
-	MessageEndpoint->Send(new FIOSLaunchDaemonLaunchApp(AppId, Params), DeviceEndpoint);
-	return true;
-#else
-	//Set return to false on Mac, since we could not find a way to do remote deploy/launch.
-	return false;
-#endif // !PLATFORM_MAC
-}
-
 bool FIOSTargetDevice::SupportsFeature(ETargetDeviceFeatures Feature) const
 {
 	switch (Feature)
@@ -443,11 +390,6 @@ bool FIOSTargetDevice::SupportsFeature(ETargetDeviceFeatures Feature) const
 	}
 }
 
-bool FIOSTargetDevice::SupportsSdkVersion(const FString& VersionString) const
-{
-	return true;
-}
-
 bool FIOSTargetDevice::TerminateProcess(const int64 ProcessId)
 {
 	return false;
diff --git a/Engine/Source/Developer/IOS/IOSTargetPlatform/Private/IOSTargetDevice.h b/Engine/Source/Developer/IOS/IOSTargetPlatform/Private/IOSTargetDevice.h
--- a/Engine/Source/Developer/IOS/IOSTargetPlatform/Private/IOSTargetDevice.h
+++ b/Engine/Source/Developer/IOS/IOSTargetPlatform/Private/IOSTargetDevice.h
@@ -109,7 +109,6 @@ public:
 	//~ ITargetDevice interface
 
 	virtual bool Connect() override;
-	virtual bool Deploy(const FString& SourceFolder, FString& OutAppId) override;
 	virtual void Disconnect() override;
 	virtual int32 GetProcessSnapshot(TArray<FTargetDeviceProcessInfo>& OutProcessInfos) override;
 	virtual ETargetDeviceTypes GetDeviceType() const override;
@@ -119,13 +118,10 @@ public:
 	virtual const class ITargetPlatform& GetTargetPlatform() const override;
 	virtual bool IsConnected() override;
 	virtual bool IsDefault() const override;
-	virtual bool Launch(const FString& InAppId, EBuildConfiguration InBuildConfiguration, EBuildTargetType TargetType, const FString& Params, uint32* OutProcessId) override;
 	virtual bool PowerOff(bool Force) override;
 	virtual bool PowerOn() override;
 	virtual bool Reboot(bool bReconnect = false) override;
-	virtual bool Run(const FString& ExecutablePath, const FString& Params, uint32* OutProcessId) override;
 	virtual bool SupportsFeature(ETargetDeviceFeatures Feature) const;
-	virtual bool SupportsSdkVersion(const FString& VersionString) const override;
 	virtual bool TerminateProcess(const int64 ProcessId) override;
 	virtual void SetUserCredentials(const FString& UserName, const FString& UserPassword) override;
 	virtual bool GetUserCredentials(FString& OutUserName, FString& OutUserPassword) override;
diff --git a/Engine/Source/Developer/Linux/LinuxTargetPlatform/Private/LinuxTargetDevice.h b/Engine/Source/Developer/Linux/LinuxTargetPlatform/Private/LinuxTargetDevice.h
--- a/Engine/Source/Developer/Linux/LinuxTargetPlatform/Private/LinuxTargetDevice.h
+++ b/Engine/Source/Developer/Linux/LinuxTargetPlatform/Private/LinuxTargetDevice.h
@@ -57,38 +57,6 @@ public:
 		return true;
 	}
 
-	virtual bool Deploy(const FString& SourceFolder, FString& OutAppId) override
-	{
-#if PLATFORM_LINUX	// if running natively, support simplified, local deployment
-		OutAppId = TEXT("");
-
-		FString PlatformName = FPlatformProcess::GetBinariesSubdirectory();
-		FString DeploymentDir = FPaths::EngineIntermediateDir() / TEXT("Devices") / PlatformName;
-
-		// delete previous build
-		IFileManager::Get().DeleteDirectory(*DeploymentDir, false, true);
-
-		// copy files into device directory
-		TArray<FString> FileNames;
-
-		IFileManager::Get().FindFilesRecursive(FileNames, *SourceFolder, TEXT("*.*"), true, false);
-
-		for (int32 FileIndex = 0; FileIndex < FileNames.Num(); ++FileIndex)
-		{
-			const FString& SourceFilePath = FileNames[FileIndex];
-			FString DestFilePath = DeploymentDir + SourceFilePath.RightChop(SourceFolder.Len());
-
-			IFileManager::Get().Copy(*DestFilePath, *SourceFilePath);
-		}
-
-		return true;
-#else
-		// @todo: support deployment to a remote machine
-		STUBBED("FLinuxTargetDevice::Deploy");
-		return false;
-#endif // PLATFORM_LINUX
-	}
-
 	virtual void Disconnect( ) override
 	{ }
 
@@ -143,70 +111,12 @@ public:
 		return false;
 	}
 
-	virtual bool Launch( const FString& AppId, EBuildConfiguration BuildConfiguration, EBuildTargetType TargetType, const FString& Params, uint32* OutProcessId ) override
-	{
-#if PLATFORM_LINUX	// if running natively, support launching in place
-		// build executable path
-
-		FString PlatformName = FPlatformProcess::GetBinariesSubdirectory();
-		FString ExecutablePath = FPaths::EngineIntermediateDir() / TEXT("Devices") / PlatformName / TEXT("Engine") / TEXT("Binaries") / PlatformName;
-
-		if (TargetType == EBuildTargetType::Game)
-		{
-			ExecutablePath /= TEXT("UnrealGame");
-		}
-		else if (TargetType == EBuildTargetType::Server)
-		{
-			ExecutablePath /= TEXT("UnrealServer");
-		}
-		else if (TargetType == EBuildTargetType::Editor)
-		{
-			ExecutablePath /= TEXT("UnrealEditor");
-		}
-
-		if (BuildConfiguration != EBuildConfiguration::Development)
-		{
-			ExecutablePath += FString::Printf(TEXT("-%s-%s"), *PlatformName, LexToString(BuildConfiguration));
-		}
-
-		// launch the game
-		FProcHandle ProcessHandle = FPlatformProcess::CreateProc(*ExecutablePath, *Params, true, false, false, OutProcessId, 0, NULL, NULL);
-		if (ProcessHandle.IsValid())
-		{
-			FPlatformProcess::CloseProc(ProcessHandle);
-			return true;
-		}
-		return false;
-#else
-		// @todo: support launching on a remote machine
-		STUBBED("FLinuxTargetDevice::Launch");
-		return false;
-#endif // PLATFORM_LINUX
-	}
-
 	virtual bool Reboot( bool bReconnect = false ) override
 	{
 		STUBBED("FLinuxTargetDevice::Reboot");
 		return false;
 	}
 
-	virtual bool Run( const FString& ExecutablePath, const FString& Params, uint32* OutProcessId ) override
-	{
-#if PLATFORM_LINUX	// if running natively, support simplified, local deployment
-		FProcHandle ProcessHandle = FPlatformProcess::CreateProc(*ExecutablePath, *Params, true, false, false, OutProcessId, 0, NULL, NULL);
-		if (ProcessHandle.IsValid())
-		{
-			FPlatformProcess::CloseProc(ProcessHandle);
-			return true;
-		}
-		return false;
-#else
-		// @todo: support remote run
-		STUBBED("FLinuxTargetDevice::Run");
-		return false;
-#endif // PLATFORM_LINUX
-	}
-
 	virtual bool SupportsFeature( ETargetDeviceFeatures Feature ) const override
 	{
 		switch (Feature)
@@ -234,12 +144,6 @@ public:
 		return false;
 	}
 
-	virtual bool SupportsSdkVersion( const FString& VersionString ) const override
-	{
-		STUBBED("FLinuxTargetDevice::SupportsSdkVersion");
-		return true;
-	}
-
 	virtual void SetUserCredentials( const FString& InUserName, const FString& InUserPassword ) override
 	{
 		UserName = InUserName;
diff --git a/Engine/Source/Developer/Mac/MacTargetPlatform/Private/LocalMacTargetDevice.h b/Engine/Source/Developer/Mac/MacTargetPlatform/Private/LocalMacTargetDevice.h
--- a/Engine/Source/Developer/Mac/MacTargetPlatform/Private/LocalMacTargetDevice.h
+++ b/Engine/Source/Developer/Mac/MacTargetPlatform/Private/LocalMacTargetDevice.h
@@ -36,32 +36,6 @@ public:
 		return true;
 	}
 
-	virtual bool Deploy( const FString& SourceFolder, FString& OutAppId ) override
-	{
-		OutAppId = TEXT("");
-
-		FString PlatformName = TEXT("Mac");
-		FString DeploymentDir = FPaths::EngineIntermediateDir() / TEXT("Devices") / PlatformName;
-
-		// delete previous build
-		IFileManager::Get().DeleteDirectory(*DeploymentDir, false, true);
-
-		// copy files into device directory
-		TArray<FString> FileNames;
-
-		IFileManager::Get().FindFilesRecursive(FileNames, *SourceFolder, TEXT("*.*"), true, false);
-
-		for (int32 FileIndex = 0; FileIndex < FileNames.Num(); ++FileIndex)
-		{
-			const FString& SourceFilePath = FileNames[FileIndex];
-			FString DestFilePath = DeploymentDir + SourceFilePath.RightChop(SourceFolder.Len());
-
-			IFileManager::Get().Copy(*DestFilePath, *SourceFilePath);
-		}
-
-		return true;
-	}
-
 	virtual void Disconnect( )
 	{ }
 
@@ -106,28 +80,6 @@ public:
 		return true;
 	}
 
-	virtual bool Launch( const FString& AppId, EBuildConfiguration BuildConfiguration, EBuildTargetType TargetType, const FString& Params, uint32* OutProcessId ) override
-	{
-		// build executable path
-		FString PlatformName = TEXT("Mac");
-		FString ExecutableName = TEXT("UE4");
-		if (BuildConfiguration != EBuildConfiguration::Development)
-		{
-			ExecutableName += FString::Printf(TEXT("-%s-%s"), *PlatformName, LexToString(BuildConfiguration));
-		}
-
-		FString ExecutablePath = FPaths::EngineIntermediateDir() / TEXT("Devices") / PlatformName / TEXT("Engine") / TEXT("Binaries") / PlatformName / (ExecutableName + TEXT(".app/Contents/MacOS/") + ExecutableName);
-
-		// launch the game
-		FProcHandle ProcessHandle = FPlatformProcess::CreateProc(*ExecutablePath, *Params, true, false, false, OutProcessId, 0, NULL, NULL);
-		if (ProcessHandle.IsValid())
-		{
-			FPlatformProcess::CloseProc(ProcessHandle);
-			return true;
-		}
-		return false;
-	}
-
 	virtual bool PowerOff( bool Force ) override
 	{
 		return false;
@@ -148,17 +100,6 @@ public:
 		return true;
 	}
 
-	virtual bool Run( const FString& ExecutablePath, const FString& Params, uint32* OutProcessId ) override
-	{
-		FProcHandle ProcessHandle = FPlatformProcess::CreateProc(*ExecutablePath, *Params, true, false, false, OutProcessId, 0, NULL, NULL);
-		if (ProcessHandle.IsValid())
-		{
-			FPlatformProcess::CloseProc(ProcessHandle);
-			return true;
-		}
-		return false;
-	}
-
 	virtual bool SupportsFeature( ETargetDeviceFeatures Feature ) const override
 	{
 		switch (Feature)
@@ -177,12 +118,6 @@ public:
 		return false;
 	}
 
-	virtual bool SupportsSdkVersion( const FString& VersionString ) const override
-	{
-		// @todo filter SDK versions
-		return true;
-	}
-
 	virtual void SetUserCredentials( const FString& UserName, const FString& UserPassword ) override
 	{
 	}
diff --git a/Engine/Source/Developer/TargetDeviceServices/Private/Proxies/TargetDeviceProxy.cpp b/Engine/Source/Developer/TargetDeviceServices/Private/Proxies/TargetDeviceProxy.cpp
--- a/Engine/Source/Developer/TargetDeviceServices/Private/Proxies/TargetDeviceProxy.cpp
+++ b/Engine/Source/Developer/TargetDeviceServices/Private/Proxies/TargetDeviceProxy.cpp
@@ -258,29 +258,6 @@ bool FTargetDeviceProxy::HasVariant(FName InVariant) const
 }
 
 
-bool FTargetDeviceProxy::DeployApp(FName InVariant, const TMap<FString, FString>& Files, const FGuid& TransactionId)
-{
-	for (TMap<FString, FString>::TConstIterator It(Files); It; ++It)
-	{
-		TSharedRef<IMessageAttachment, ESPMode::ThreadSafe> FileAttachment = MakeShareable(new FFileMessageAttachment(It.Key()));
-		FString SourcePath = It.Key();
-
-		MessageEndpoint->Send(new FTargetDeviceServiceDeployFile(It.Value(), TransactionId), FileAttachment, MessageAddress);
-	}
-
-	MessageEndpoint->Send(new FTargetDeviceServiceDeployCommit(InVariant, TransactionId), MessageAddress);
-
-	return true;
-}
-
-
-bool FTargetDeviceProxy::LaunchApp(FName InVariant, const FString& AppId, EBuildConfiguration BuildConfiguration, const FString& Params)
-{
-	MessageEndpoint->Send(new FTargetDeviceServiceLaunchApp(InVariant, AppId, BuildConfiguration, Params), MessageAddress);
-
-	return true;
-}
-
 
 bool FTargetDeviceProxy::TerminateLaunchedProcess(FName InVariant, const FString& ProcessIdentifier)
 {
@@ -308,46 +285,11 @@ void FTargetDeviceProxy::Reboot()
 }
 
 
-void FTargetDeviceProxy::Run(FName InVariant, const FString& ExecutablePath, const FString& Params)
-{
-	MessageEndpoint->Send(new FTargetDeviceServiceRunExecutable(InVariant, ExecutablePath, Params), MessageAddress);
-}
-
 
 /* FTargetDeviceProxy implementation
  *****************************************************************************/
 
 void FTargetDeviceProxy::InitializeMessaging()
 {
-	MessageEndpoint = FMessageEndpoint::Builder(FName(*FString::Printf(TEXT("FTargetDeviceProxy (%s)"), *Name)))
-		.Handling<FTargetDeviceServiceDeployFinished>(this, &FTargetDeviceProxy::HandleDeployFinishedMessage)
-		.Handling<FTargetDeviceServiceLaunchFinished>(this, &FTargetDeviceProxy::HandleLaunchFinishedMessage);
-}
-
-/* FTargetDeviceProxy event handlers
-*****************************************************************************/
-
-void FTargetDeviceProxy::HandleDeployFinishedMessage(const FTargetDeviceServiceDeployFinished& Message, const TSharedRef<IMessageContext, ESPMode::ThreadSafe>& Context)
-{
-	if (Message.Succeeded)
-	{
-		DeployCommittedDelegate.Broadcast(Message.TransactionId, Message.AppID);
-	}
-	else
-	{
-		DeployFailedDelegate.Broadcast(Message.TransactionId);
-	}
-}
-
-
-void FTargetDeviceProxy::HandleLaunchFinishedMessage(const FTargetDeviceServiceLaunchFinished& Message, const TSharedRef<IMessageContext, ESPMode::ThreadSafe>& Context)
-{
-	if (Message.Succeeded)
-	{
-		LaunchSucceededDelegate.Broadcast(Message.AppID, Message.ProcessId);
-	}
-	else
-	{
-		LaunchFailedDelegate.Broadcast(Message.AppID);
-	}
+	// could delete this function
 }
diff --git a/Engine/Source/Developer/TargetDeviceServices/Private/Proxies/TargetDeviceProxy.h b/Engine/Source/Developer/TargetDeviceServices/Private/Proxies/TargetDeviceProxy.h
--- a/Engine/Source/Developer/TargetDeviceServices/Private/Proxies/TargetDeviceProxy.h
+++ b/Engine/Source/Developer/TargetDeviceServices/Private/Proxies/TargetDeviceProxy.h
@@ -174,31 +174,12 @@ public:
 	{
 		return Shared;
 	}
-
-	virtual bool DeployApp(FName InVariant, const TMap<FString, FString>& Files, const FGuid& TransactionId) override;
-	virtual bool LaunchApp(FName InVariant, const FString& AppId, EBuildConfiguration BuildConfiguration, const FString& Params) override;
 	
 	virtual bool TerminateLaunchedProcess(FName InVariant, const FString& ProcessIdentifier) override;
 
-	virtual FOnTargetDeviceProxyDeployCommitted& OnDeployCommitted() override
-	{
-		return DeployCommittedDelegate;
-	}
-
-	virtual FOnTargetDeviceProxyDeployFailed& OnDeployFailed() override
-	{
-		return DeployFailedDelegate;
-	}
-
-	virtual FOnTargetDeviceProxyLaunchFailed& OnLaunchFailed() override
-	{
-		return LaunchFailedDelegate;
-	}
-
 	virtual void PowerOff(bool Force) override;
 	virtual void PowerOn() override;
 	virtual void Reboot() override;
-	virtual void Run(FName InVariant, const FString& ExecutablePath, const FString& Params) override;
 
 	/** Returns true if this is an aggregate (All_<platform>_devices_on_<host>) proxy, false otherwise */
 	virtual bool IsAggregated() const override
@@ -211,14 +192,6 @@ protected:
 	/** Initializes the message endpoint. */
 	void InitializeMessaging();
 
-private:
-
-	/** Handles FTargetDeviceServiceDeployFinishedMessage messages. */
-	void HandleDeployFinishedMessage(const FTargetDeviceServiceDeployFinished& Message, const TSharedRef<IMessageContext, ESPMode::ThreadSafe>& Context);
-
-	/** Handles FTargetDeviceServiceLaunchFinishedMessage messages. */
-	void HandleLaunchFinishedMessage(const FTargetDeviceServiceLaunchFinished& Message, const TSharedRef<IMessageContext, ESPMode::ThreadSafe>& Context);
-
 private:
 
 	/** Holds a flag indicating whether the device is connected. */
@@ -305,18 +278,4 @@ private:
 
 	/** Map of all the Variants for this Device. */
 	TMap<FName, FTargetDeviceProxyVariant> TargetDeviceVariants;
-
-private:
-
-	/** Holds a delegate to be invoked when a build has been deployed to the target device. */
-	FOnTargetDeviceProxyDeployCommitted DeployCommittedDelegate;
-
-	/** Holds a delegate to be invoked when a build has failed to deploy to the target device. */
-	FOnTargetDeviceProxyDeployFailed DeployFailedDelegate;
-
-	/** Holds a delegate to be invoked when a build has failed to launch on the target device. */
-	FOnTargetDeviceProxyLaunchFailed LaunchFailedDelegate;
-
-	/** Holds a delegate to be invoked when a build has succeeded to launch on the target device. */
-	FOnTargetDeviceProxyLaunchSucceeded LaunchSucceededDelegate;
 };
diff --git a/Engine/Source/Developer/TargetDeviceServices/Private/Services/TargetDeviceService.cpp b/Engine/Source/Developer/TargetDeviceServices/Private/Services/TargetDeviceService.cpp
--- a/Engine/Source/Developer/TargetDeviceServices/Private/Services/TargetDeviceService.cpp
+++ b/Engine/Source/Developer/TargetDeviceServices/Private/Services/TargetDeviceService.cpp
@@ -44,15 +44,11 @@ FTargetDeviceService::FTargetDeviceService(const FString& InDeviceName, const TS
 	MessageEndpoint = FMessageEndpoint::Builder(FName(*FString::Printf(TEXT("FTargetDeviceService (%s)"), *DeviceName)), InMessageBus)
 		.Handling<FTargetDeviceClaimDenied>(this, &FTargetDeviceService::HandleClaimDeniedMessage)
 		.Handling<FTargetDeviceClaimed>(this, &FTargetDeviceService::HandleClaimedMessage)
-		.Handling<FTargetDeviceServiceDeployCommit>(this, &FTargetDeviceService::HandleDeployCommitMessage)
-		.Handling<FTargetDeviceServiceDeployFile>(this, &FTargetDeviceService::HandleDeployFileMessage)
-		.Handling<FTargetDeviceServiceLaunchApp>(this, &FTargetDeviceService::HandleLaunchAppMessage)
 		.Handling<FTargetDeviceServiceTerminateLaunchedProcess>(this, &FTargetDeviceService::HandleTerminateLaunchedProcessMessage)
 		.Handling<FTargetDeviceServicePing>(this, &FTargetDeviceService::HandlePingMessage)
 		.Handling<FTargetDeviceServicePowerOff>(this, &FTargetDeviceService::HandlePowerOffMessage)
 		.Handling<FTargetDeviceServicePowerOn>(this, &FTargetDeviceService::HandlePowerOnMessage)
 		.Handling<FTargetDeviceServiceReboot>(this, &FTargetDeviceService::HandleRebootMessage)
-		.Handling<FTargetDeviceServiceRunExecutable>(this, &FTargetDeviceService::HandleRunExecutableMessage)
 		.Handling<FTargetDeviceUnclaimed>(this, &FTargetDeviceService::HandleUnclaimedMessage);
 
 	if (MessageEndpoint.IsValid())
@@ -388,55 +384,6 @@ void FTargetDeviceService::HandleUnclaimedMessage(const FTargetDeviceUnclaimed&
 }
 
 
-void FTargetDeviceService::HandleDeployFileMessage(const FTargetDeviceServiceDeployFile& Message, const TSharedRef<IMessageContext, ESPMode::ThreadSafe>& Context)
-{
-	if (!Running)
-	{
-		return;
-	}
-
-	TSharedPtr<IMessageAttachment, ESPMode::ThreadSafe> Attachment = Context->GetAttachment();
-
-	if (Attachment.IsValid())
-	{
-		FArchive* FileReader = Attachment->CreateReader();
-
-		if (FileReader != nullptr)
-		{
-			FString DeploymentFolder = FPaths::EngineIntermediateDir() / TEXT("Deploy") / Message.TransactionId.ToString();
-			FString TargetPath = DeploymentFolder / Message.TargetFileName;
-
-			StoreDeployedFile(FileReader, TargetPath);
-
-			delete FileReader;
-		}
-	}
-}
-
-
-void FTargetDeviceService::HandleDeployCommitMessage(const FTargetDeviceServiceDeployCommit& Message, const TSharedRef<IMessageContext, ESPMode::ThreadSafe>& Context)
-{
-	if (!Running)
-	{
-		return;
-	}
-
-	ITargetDevicePtr TargetDevice = GetDevice(Message.Variant);
-
-	if (TargetDevice.IsValid())
-	{
-		FString SourceFolder = FPaths::EngineIntermediateDir() / TEXT("Deploy") / Message.TransactionId.ToString();
-		FString OutAppId;
-
-		bool Succeeded = TargetDevice->Deploy(SourceFolder, OutAppId);
-
-		IFileManager::Get().DeleteDirectory(*SourceFolder, false, true);
-		// message is going to be deleted by FMemory::Free() (see FMessageContext destructor), so allocate it with Malloc
-		void* Memory = FMemory::Malloc(sizeof(FTargetDeviceServiceDeployFinished), alignof(FTargetDeviceServiceDeployFinished));
-		MessageEndpoint->Send(new(Memory) FTargetDeviceServiceDeployFinished(Message.Variant, OutAppId, Succeeded, Message.TransactionId), Context->GetSender());
-	}
-}
-
 void FTargetDeviceService::HandleTerminateLaunchedProcessMessage(const FTargetDeviceServiceTerminateLaunchedProcess& Message, const TSharedRef<IMessageContext, ESPMode::ThreadSafe>& Context)
 {
 	if (!Running)
@@ -452,30 +399,6 @@ void FTargetDeviceService::HandleTerminateLaunchedProcessMessage(const FTargetDe
 	}
 }
 
-void FTargetDeviceService::HandleLaunchAppMessage(const FTargetDeviceServiceLaunchApp& Message, const TSharedRef<IMessageContext, ESPMode::ThreadSafe>& Context)
-{
-	if (!Running)
-	{
-		return;
-	}
-
-	ITargetDevicePtr TargetDevice = GetDevice(Message.Variant);
-
-	if (TargetDevice.IsValid())
-	{
-		uint32 ProcessId;
-		bool Succeeded = TargetDevice->Launch(Message.AppID, (EBuildConfiguration)Message.BuildConfiguration, EBuildTargetType::Game, Message.Params, &ProcessId);
-
-		if (MessageEndpoint.IsValid())
-		{
-			// message is going to be deleted by FMemory::Free() (see FMessageContext destructor), so allocate it with Malloc
-			void* Memory = FMemory::Malloc(sizeof(FTargetDeviceServiceLaunchFinished), alignof(FTargetDeviceServiceLaunchFinished));
-			MessageEndpoint->Send(new(Memory) FTargetDeviceServiceLaunchFinished(Message.AppID, ProcessId, Succeeded), Context->GetSender());
-		}
-	}
-}
-
-
 void FTargetDeviceService::HandlePingMessage(const FTargetDeviceServicePing& InMessage, const TSharedRef<IMessageContext, ESPMode::ThreadSafe>& Context)
 {
 	if (!Running)
@@ -591,24 +514,3 @@ void FTargetDeviceService::HandleRebootMessage( const FTargetDeviceServiceReboot
 		TargetDevice->Reboot();
 	}
 }
-
-
-void FTargetDeviceService::HandleRunExecutableMessage(const FTargetDeviceServiceRunExecutable& Message, const TSharedRef<IMessageContext, ESPMode::ThreadSafe>& Context)
-{
-	if (!Running)
-	{
-		return;
-	}
-
-	ITargetDevicePtr TargetDevice = GetDevice(Message.Variant);
-
-	if (TargetDevice.IsValid())
-	{
-		uint32 OutProcessId;
-		bool Succeeded = TargetDevice->Run(Message.ExecutablePath, Message.Params, &OutProcessId);
-
-		// message is going to be deleted by FMemory::Free() (see FMessageContext destructor), so allocate it with Malloc
-		void* Memory = FMemory::Malloc(sizeof(FTargetDeviceServiceRunFinished), alignof(FTargetDeviceServiceRunFinished));
-		MessageEndpoint->Send(new(Memory) FTargetDeviceServiceRunFinished(Message.Variant, Message.ExecutablePath, OutProcessId, Succeeded), Context->GetSender());
-	}
-}
diff --git a/Engine/Source/Developer/TargetDeviceServices/Private/Services/TargetDeviceService.h b/Engine/Source/Developer/TargetDeviceServices/Private/Services/TargetDeviceService.h
--- a/Engine/Source/Developer/TargetDeviceServices/Private/Services/TargetDeviceService.h
+++ b/Engine/Source/Developer/TargetDeviceServices/Private/Services/TargetDeviceService.h
@@ -90,15 +90,6 @@ private:
 	/** Callback for FTargetDeviceClaimDropped messages. */
 	void HandleUnclaimedMessage(const FTargetDeviceUnclaimed& Message, const TSharedRef<IMessageContext, ESPMode::ThreadSafe>& Context);
 
-	/** Callback for FTargetDeviceServiceDeployFile messages. */
-	void HandleDeployFileMessage(const FTargetDeviceServiceDeployFile& Message, const TSharedRef<IMessageContext, ESPMode::ThreadSafe>& Context);
-
-	/** Callback for FTargetDeviceServiceDeployCommit messages. */
-	void HandleDeployCommitMessage(const FTargetDeviceServiceDeployCommit& Message, const TSharedRef<IMessageContext, ESPMode::ThreadSafe>& Context);
-
-	/** Callback for FTargetDeviceServiceLaunchApp messages. */
-	void HandleLaunchAppMessage(const FTargetDeviceServiceLaunchApp& Message, const TSharedRef<IMessageContext, ESPMode::ThreadSafe>& Context);
-
 	/** Callback for FTargetDeviceServiceTerminateLaunchedProcess messages. */
 	void HandleTerminateLaunchedProcessMessage(const FTargetDeviceServiceTerminateLaunchedProcess& Message, const TSharedRef<IMessageContext, ESPMode::ThreadSafe>& Context);
 
@@ -114,9 +105,6 @@ private:
 	/** Callback for FTargetDeviceServiceReboot messages. */
 	void HandleRebootMessage(const FTargetDeviceServiceReboot& Message, const TSharedRef<IMessageContext, ESPMode::ThreadSafe>& Context);
 
-	/** Callback for FTargetDeviceServiceRunExecutable messages. */
-	void HandleRunExecutableMessage(const FTargetDeviceServiceRunExecutable& Message, const TSharedRef<IMessageContext, ESPMode::ThreadSafe>& Context);
-
 private:
 
 	/** Default target device used when no flavor is specified. */
diff --git a/Engine/Source/Developer/TargetDeviceServices/Public/ITargetDeviceProxy.h b/Engine/Source/Developer/TargetDeviceServices/Public/ITargetDeviceProxy.h
--- a/Engine/Source/Developer/TargetDeviceServices/Public/ITargetDeviceProxy.h
+++ b/Engine/Source/Developer/TargetDeviceServices/Public/ITargetDeviceProxy.h
@@ -23,37 +23,6 @@ UE_DEPRECATED(4.16, "ITargetDeviceProxyRef is deprecated. Please use 'TSharedPtr
 typedef TSharedRef<class ITargetDeviceProxy> ITargetDeviceProxyRef;
 
 
-/**
- * Delegate type for committed build deployments.
- *
- * The first parameter is the transaction identifier of the completed deployment.
- * The second parameter is the identifier of the deployed application (used for launching).
- */
-DECLARE_MULTICAST_DELEGATE_TwoParams(FOnTargetDeviceProxyDeployCommitted, const FGuid&, const FString&)
-
-/**
- * Delegate type for failed build deployments.
- *
- * The first parameter is the transaction identifier of the failed deployment.
- */
-DECLARE_MULTICAST_DELEGATE_OneParam(FOnTargetDeviceProxyDeployFailed, const FGuid&)
-
-/**
- * Delegate type for failed build launches.
- *
- * The first parameter is the application identifier of the failed launch.
- */
-DECLARE_MULTICAST_DELEGATE_OneParam(FOnTargetDeviceProxyLaunchFailed, const FString&)
-
-/**
- * Delegate type for successful build launches.
- *
- * The first parameter is the application identifier of the successful launch.
- * The second parameter is the process identifier of the launched application.
- */
-DECLARE_MULTICAST_DELEGATE_TwoParams(FOnTargetDeviceProxyLaunchSucceeded, const FString&, uint32)
-
-
 /**
  * Interface for target device proxies.
  */
@@ -285,34 +254,6 @@ public:
 
 public:
 
-	/**
-	 * Deploys the specified list of files to the device.
-	 *
-	 * If the deployment succeeds, the OnDeployCommitted() delegate will return an
-	 * application identifier (AppId) that can be used to launch to the deployed
-	 * application. If deploymend fails, the OnDeployFailed() delegate will be invoked.
-	 *
-	 * @param InVariant Variant to deploy files on.
-	 * @param Files The files to deploy.
-	 * @param TransactionId The transaction identifier for the deployment.
-	 * @return true if deployment has been started, false otherwise.
-	 * @see Launch, OnDeployCommitted, OnDeployFailed
-	 */
-	virtual bool DeployApp(FName InVariant, const TMap<FString, FString>& Files, const FGuid& TransactionId) = 0;
-
-	/**
-	 * Launches a previously deployed build on the target device.
-	 *
-	 * If the launching fails, the OnLaunchFailed() delegate will be invoked.
-	 *
-	 * @param InVariant Variant to launch the app with.
-	 * @param AppId The identifier of the application to launch.
-	 * @param BuildConfiguration The build configuration to launch with.
-	 * @param Params The command line parameters to launch with.
-	 * @return true on success, false otherwise.
-	 * @see Deploy, OnLaunchFailed
-	 */
-	virtual bool LaunchApp(FName InVariant, const FString& AppId, EBuildConfiguration BuildConfiguration, const FString& Params) = 0;
 
 	/**
 	 * Powers off the device.
@@ -336,41 +277,6 @@ public:
 	 */
 	virtual void Reboot() = 0;
 
-	/**
-	 * Attempts to run an executable on the device.
-	 *
-	 * @param InVariant Variant to run the executable with.
-	 * @param ExecutablePath The path to the executable to run.
-	 * @param Params The command line parameters.
-	 */
-	virtual void Run(FName InVariant, const FString& ExecutablePath, const FString& Params) = 0;
-
-public:
-
-	/**
-	 * Returns a delegate to be invoked when a build has been deployed to the target device.
-	 *
-	 * @return The delegate.
-	 * @see DeployApp
-	 */
-	virtual FOnTargetDeviceProxyDeployCommitted& OnDeployCommitted() = 0;
-
-	/**
-	 * Returns a delegate to be invoked when a build failed to deploy to the target device.
-	 *
-	 * @return The delegate.
-	 * @see DeployApp
-	 */
-	virtual FOnTargetDeviceProxyDeployFailed& OnDeployFailed() = 0;
-
-	/**
-	 * Returns a delegate to be invoked when a build failed to launch on the target device.
-	 *
-	 * @return The delegate.
-	 * @see LaunchApp
-	 */
-	virtual FOnTargetDeviceProxyLaunchFailed& OnLaunchFailed() = 0;
-
 public:
 
 	/** Virtual destructor. */
diff --git a/Engine/Source/Developer/TargetDeviceServices/Public/TargetDeviceServiceMessages.h b/Engine/Source/Developer/TargetDeviceServices/Public/TargetDeviceServiceMessages.h
--- a/Engine/Source/Developer/TargetDeviceServices/Public/TargetDeviceServiceMessages.h
+++ b/Engine/Source/Developer/TargetDeviceServices/Public/TargetDeviceServiceMessages.h
@@ -12,210 +12,6 @@
 #include "TargetDeviceServiceMessages.generated.h"
 
 
-/* Application deployment messages
- *****************************************************************************/
-
-/**
- * Implements a message for committing a deployment transaction.
- *
- * @see FTargetDeviceServiceDeployFile, FTargetDeviceServiceDeployFinished
- */
-USTRUCT()
-struct FTargetDeviceServiceDeployCommit
-{
-	GENERATED_USTRUCT_BODY()
-
-	/** Holds the variant identifier of the target device to use. */
-	UPROPERTY(EditAnywhere, Category="Message")
-	FName Variant;
-
-	/** Holds the identifier of the deployment transaction to commit. */
-	UPROPERTY(EditAnywhere, Category="Message")
-	FGuid TransactionId;
-
-
-	/** Default constructor. */
-	FTargetDeviceServiceDeployCommit() { }
-
-	/** Creates and initializes a new instance. */
-	FTargetDeviceServiceDeployCommit(FName InVariant, const FGuid& InTransactionId)
-		: Variant(InVariant)
-		, TransactionId(InTransactionId)
-	{ }
-};
-
-
-/**
- * Implements a message for deploying a single file to a target device.
- *
- * The actual file data must be attached to the message.
- *
- * @see FTargetDeviceServiceDeployCommit, FTargetDeviceServiceDeployFinished
- */
-USTRUCT()
-struct FTargetDeviceServiceDeployFile
-{
-	GENERATED_USTRUCT_BODY()
-
-	/** Holds the name and path of the file as it will be stored on the target device. */
-	UPROPERTY(EditAnywhere, Category="Message")
-	FString TargetFileName;
-
-	/** Holds the identifier of the deployment transaction that this file is part of. */
-	UPROPERTY(EditAnywhere, Category="Message")
-	FGuid TransactionId;
-
-
-	/** Default constructor. */
-	FTargetDeviceServiceDeployFile() { }
-
-	/** Creates and initializes a new instance. */
-	FTargetDeviceServiceDeployFile(const FString& InTargetFileName, const FGuid& InTransactionId)
-		: TargetFileName(InTargetFileName)
-		, TransactionId(InTransactionId)
-	{ }
-};
-
-
-/**
- * Implements a message for notifying a target device proxy that a deployment transaction has finished.
- *
- * @see FTargetDeviceServiceDeployFile, FTargetDeviceServiceDeployCommit
- */
-USTRUCT()
-struct FTargetDeviceServiceDeployFinished
-{
-	GENERATED_USTRUCT_BODY()
-
-	/** Holds the variant identifier of the target device to use. */
-	UPROPERTY(EditAnywhere, Category="Message")
-	FName Variant;
-
-	/**
-	 * Holds the created identifier for the deployed application.
-	 *
-	 * The semantics of this identifier are target platform specific. In some cases it may be
-	 * a GUID, in other cases it may be the path to the application or some other means of
-	 * identifying the application. Application identifiers are returned from target device
-	 * services as result of successful deployment transactions.
-	 */
-	UPROPERTY(EditAnywhere, Category="Message")
-	FString AppID;
-
-	/** Holds a flag indicating whether the deployment transaction finished successfully. */
-	UPROPERTY(EditAnywhere, Category="Message")
-	bool Succeeded;
-
-	/** Holds the identifier of the deployment transaction that this file is part of. */
-	UPROPERTY(EditAnywhere, Category="Message")
-	FGuid TransactionId;
-
-
-	/** Default constructor. */
-	FTargetDeviceServiceDeployFinished() : Succeeded(false) { }
-
-	/** Creates and initializes a new instance. */
-	FTargetDeviceServiceDeployFinished(FName InVariant, const FString& InAppId, bool InSucceeded, FGuid InTransactionId)
-		: Variant(InVariant)
-		, AppID(InAppId)
-		, Succeeded(InSucceeded)
-		, TransactionId(InTransactionId)
-	{ }
-};
-
-/* Application launch messages
-*****************************************************************************/
-
-/**
- * Implements a message for committing a deployment transaction.
- *
- * To launch an arbitrary executable on a device use the FTargetDeviceServiceRunExecutable message instead.
- *
- * @see FTargetDeviceServiceLaunchFinished
- */
-USTRUCT()
-struct FTargetDeviceServiceLaunchApp
-{
-	GENERATED_USTRUCT_BODY()
-
-	/** Holds the variant identifier of the target device to use. */
-	UPROPERTY(EditAnywhere, Category="Message")
-	FName Variant;
-
-	/**
-	 * Holds the identifier of the application to launch.
-	 *
-	 * The semantics of this identifier are target platform specific. In some cases it may be
-	 * a GUID, in other cases it may be the path to the application or some other means of
-	 * identifying the application. Application identifiers are returned from target device
-	 * services as result of successful deployment transactions.
-	 */
-	UPROPERTY(EditAnywhere, Category="Message")
-	FString AppID;
-
-	/** The application's build configuration, i.e. Debug or Shipping. */
-	UPROPERTY(EditAnywhere, Category="Message")
-	uint8 BuildConfiguration;
-
-	/** Holds optional command line parameters for the application. */
-	UPROPERTY(EditAnywhere, Category="Message")
-	FString Params;
-
-
-	/** Default constructor. */
-	FTargetDeviceServiceLaunchApp() : BuildConfiguration(0) { }
-
-	/** Creates and initializes a new instance. */
-	FTargetDeviceServiceLaunchApp(FName InVariant, const FString& InAppId, EBuildConfiguration InBuildConfiguration, const FString& InParams)
-		: Variant(InVariant)
-		, AppID(InAppId)
-		, BuildConfiguration((uint8)InBuildConfiguration)
-		, Params(InParams)
-	{ }
-};
-
-
-/**
- * Implements a message for notifying a target device proxy that launching an application has finished.
- *
- * @see FTargetDeviceServiceLaunchApp
- */
-USTRUCT()
-struct FTargetDeviceServiceLaunchFinished
-{
-	GENERATED_USTRUCT_BODY()
-
-	/**
-	 * Holds the identifier of the launched application.
-	 *
-	 * The semantics of this identifier are target platform specific. In some cases it may be
-	 * a GUID, in other cases it may be the path to the application or some other means of
-	 * identifying the application. Application identifiers are returned from target device
-	 * services as result of successful deployment transactions.
-	 */
-	UPROPERTY(EditAnywhere, Category="Message")
-	FString AppID;
-
-	/** Holds the process identifier for the launched application. */
-	UPROPERTY(EditAnywhere, Category="Message")
-	int32 ProcessId;
-
-	/** Holds a flag indicating whether the application was launched successfully. */
-	UPROPERTY(EditAnywhere, Category="Message")
-	bool Succeeded;
-
-
-	/** Default constructor. */
-	FTargetDeviceServiceLaunchFinished() : ProcessId(0), Succeeded(false) { }
-
-	/** Creates and initializes a new instance. */
-	FTargetDeviceServiceLaunchFinished(const FString& InAppId, int32 InProcessId, bool InSucceeded)
-		: AppID(InAppId)
-		, ProcessId(InProcessId)
-		, Succeeded(InSucceeded)
-	{ }
-};
-
 
 /* Application terminate process messages
 *****************************************************************************/
@@ -593,83 +389,3 @@ struct FTargetDeviceServiceReboot
 		: Operator(InOperator)
 	{ }
 };
-
-
-/**
- * Implements a message for running an executable on a target device.
- *
- * To launch a previously deployed application on a device use the FTargetDeviceServiceLaunchApp message instead.
- *
- * @see FTargetDeviceServiceLaunchApp
- */
-USTRUCT()
-struct FTargetDeviceServiceRunExecutable
-{
-	GENERATED_USTRUCT_BODY()
-
-	/** Holds the variant identifier of the target device to use for execution. */
-	UPROPERTY(EditAnywhere, Category="Message")
-	FName Variant;
-
-	/** Holds the path to the executable on the device. */
-	UPROPERTY(EditAnywhere, Category="Message")
-	FString ExecutablePath;
-
-	/** Holds optional command line parameters for the executable. */
-	UPROPERTY(EditAnywhere, Category="Message")
-	FString Params;
-
-
-	/** Default constructor. */
-	FTargetDeviceServiceRunExecutable() { }
-
-	/** Creates and initializes a new instance. */
-	FTargetDeviceServiceRunExecutable(FName InVariant, const FString& InExecutablePath, const FString& InParams)
-		: Variant(InVariant)
-		, ExecutablePath(InExecutablePath)
-		, Params(InParams)
-	{ }
-};
-
-
-/**
- * Implements a message for notifying a target device proxy that running an executable has finished.
- *
- * @see FTargetDeviceServiceRunExecutable
- */
-USTRUCT()
-struct FTargetDeviceServiceRunFinished
-{
-	GENERATED_USTRUCT_BODY()
-
-	/** Holds the variant identifier of the target device to use. */
-	UPROPERTY(EditAnywhere, Category="Message")
-	FName Variant;
-
-	/** Holds the path to the executable that was run. */
-	UPROPERTY(EditAnywhere, Category="Message")
-	FString ExecutablePath;
-
-	/** Holds the process identifier of the running executable. */
-	UPROPERTY(EditAnywhere, Category="Message")
-	int32 ProcessId;
-
-	/** Holds a flag indicating whether the executable started successfully. */
-	UPROPERTY(EditAnywhere, Category="Message")
-	bool Succeeded;
-
-
-	/** Default constructor. */
-	FTargetDeviceServiceRunFinished()
-		: ProcessId(0)
-		, Succeeded(false)
-	{ }
-
-	/** Creates and initializes a new instance. */
-	FTargetDeviceServiceRunFinished(FName InVariant, const FString& InExecutablePath, int32 InProcessId, bool InSucceeded)
-		: Variant(InVariant)
-		, ExecutablePath(InExecutablePath)
-		, ProcessId(InProcessId)
-		, Succeeded(InSucceeded)
-	{ }
-};
diff --git a/Engine/Source/Developer/TargetPlatform/Public/Interfaces/ITargetDevice.h b/Engine/Source/Developer/TargetPlatform/Public/Interfaces/ITargetDevice.h
--- a/Engine/Source/Developer/TargetPlatform/Public/Interfaces/ITargetDevice.h
+++ b/Engine/Source/Developer/TargetPlatform/Public/Interfaces/ITargetDevice.h
@@ -215,15 +215,6 @@ public:
 	 */
 	virtual bool Connect() = 0;
 
-	/**
-	 * Deploys an application in the specified folder to the device.
-	 *
-	 * @param SourceFolder The path to the files and directories to be deployed.
-	 * @param OutAppId Will hold the identifier of the deployed application (used for launching).
-	 * @return true on success, false otherwise.
-	 */
-	virtual bool Deploy( const FString& SourceFolder, FString& OutAppId ) = 0;
-
 	/**
 	 * Disconnect from the physical device.
 	 */
@@ -317,18 +308,6 @@ public:
 	*/
 	virtual bool IsAuthorized() const { return true;  }
 
-	/**
-	 * Launches a previously deployed build.
-	 *
-	 * @param AppId The identifier of the application to launch (as returned by the Deploy() method).
-	 * @param BuildConfiguration The build configuration to launch.
-	 * @param TargetType The target type to launch
-	 * @param Params The command line parameters to launch with.
-	 * @param OutProcessId Will hold the identifier of the created process (can be NULL).
-	 * @return true on success, false otherwise.
-	 */
-	virtual bool Launch( const FString& AppId, EBuildConfiguration BuildConfiguration, EBuildTargetType TargetType, const FString& Params, uint32* OutProcessId ) = 0;
-
 	/**
 	 * Powers off the device.
 	 *
@@ -352,16 +331,6 @@ public:
 	 */
 	virtual bool Reboot( bool bReconnect = false ) = 0;
 
-	/**
-	 * Runs an executable on the device.
-	 *
-	 * @param ExecutablePath The path to the executable to run.
-	 * @param Params The command line parameters.
-	 * @param OutProcessId Will hold the identifier of the created process (can be NULL).
-	 * @return true if the executable was started, false otherwise.
-	 */
-	virtual bool Run( const FString& ExecutablePath, const FString& Params, uint32* OutProcessId ) = 0;
-
 	/**
 	 * Checks whether the target device supports the specified feature.
 	 *
@@ -370,14 +339,6 @@ public:
 	 */
 	virtual bool SupportsFeature( ETargetDeviceFeatures Feature ) const = 0;
 
-	/**
-	 * Checks whether this device supports the specified SDK version.
-	 *
-	 * @param VersionString The SDK version string.
-	 * @return true if the SDK version is supported, false otherwise.
-	 */
-	virtual bool SupportsSdkVersion( const FString& VersionString ) const = 0;
-
 	/**
 	 * Terminates a process that was launched on the device using the Launch() or Run() methods.
 	 *
diff --git a/Engine/Source/Developer/Windows/WindowsTargetPlatform/Private/LocalPcTargetDevice.h b/Engine/Source/Developer/Windows/WindowsTargetPlatform/Private/LocalPcTargetDevice.h
--- a/Engine/Source/Developer/Windows/WindowsTargetPlatform/Private/LocalPcTargetDevice.h
+++ b/Engine/Source/Developer/Windows/WindowsTargetPlatform/Private/LocalPcTargetDevice.h
@@ -40,32 +40,6 @@ public:
 		return true;
 	}
 
-	virtual bool Deploy( const FString& SourceFolder, FString& OutAppId ) override
-	{
-		OutAppId = TEXT("");
-
-		FString PlatformName = WIN64 ? TEXT("Win64") : TEXT("Win32");
-		FString DeploymentDir = FPaths::EngineIntermediateDir() / TEXT("Devices") / PlatformName;
-
-		// delete previous build
-		IFileManager::Get().DeleteDirectory(*DeploymentDir, false, true);
-
-		// copy files into device directory
-		TArray<FString> FileNames;
-
-		IFileManager::Get().FindFilesRecursive(FileNames, *SourceFolder, TEXT("*.*"), true, false);
-
-		for (int32 FileIndex = 0; FileIndex < FileNames.Num(); ++FileIndex)
-		{
-			const FString& SourceFilePath = FileNames[FileIndex];
-			FString DestFilePath = DeploymentDir + SourceFilePath.RightChop(SourceFolder.Len());
-
-			IFileManager::Get().Copy(*DestFilePath, *SourceFilePath);
-		}
-
-		return true;
-	}
-
 	virtual void Disconnect( )
 	{ }
 
@@ -237,42 +211,6 @@ public:
 		return true;
 	}
 
-	virtual bool Launch( const FString& AppId, EBuildConfiguration BuildConfiguration, EBuildTargetType TargetType, const FString& Params, uint32* OutProcessId ) override
-	{
-		// build executable path
-		FString PlatformName = WIN64 ? TEXT("Win64") : TEXT("Win32");
-		FString ExecutablePath = FPaths::EngineIntermediateDir() / TEXT("Devices") / PlatformName / TEXT("Engine") / TEXT("Binaries") / PlatformName;
-		
-		if (TargetType == EBuildTargetType::Game)
-		{
-			ExecutablePath /= TEXT("UnrealGame");
-		}
-		else if (TargetType == EBuildTargetType::Server)
-		{
-			ExecutablePath /= TEXT("UnrealGame");
-		}
-		else if (TargetType == EBuildTargetType::Editor)
-		{
-			ExecutablePath /= TEXT("UnrealEditor");
-		}
-
-		if (BuildConfiguration != EBuildConfiguration::Development)
-		{
-			ExecutablePath += FString::Printf(TEXT("-%s-%s"), *PlatformName, LexToString(BuildConfiguration));
-		}
-
-		ExecutablePath += TEXT(".exe");
-
-		// launch the game
-		FProcHandle ProcessHandle = FPlatformProcess::CreateProc(*ExecutablePath, *Params, true, false, false, OutProcessId, 0, nullptr, nullptr);
-		if (ProcessHandle.IsValid())
-		{
-			FPlatformProcess::CloseProc(ProcessHandle);
-			return true;
-		}
-		return false;
-	}
-
 	virtual bool PowerOff( bool Force ) override
 	{
 		if (!AdjustShutdownPrivileges())
@@ -300,17 +238,6 @@ public:
 		return (::ExitWindowsEx(EWX_REBOOT | EWX_FORCE, SHTDN_REASON_MINOR_MAINTENANCE | SHTDN_REASON_FLAG_PLANNED) != 0);
 	}
 
-	virtual bool Run( const FString& ExecutablePath, const FString& Params, uint32* OutProcessId )
-	{
-		FProcHandle ProcessHandle = FPlatformProcess::CreateProc(*ExecutablePath, *Params, true, false, false, OutProcessId, 0, nullptr, nullptr);
-		if (ProcessHandle.IsValid())
-		{
-			FPlatformProcess::CloseProc(ProcessHandle);
-			return true;
-		}
-		return false;
-	}
-
 	virtual bool SupportsFeature( ETargetDeviceFeatures Feature ) const override
 	{
 		switch (Feature)
@@ -335,12 +262,6 @@ public:
 		return false;
 	}
 
-	virtual bool SupportsSdkVersion( const FString& VersionString ) const override
-	{
-		// @todo filter SDK versions
-		return true;
-	}
-
 	virtual void SetUserCredentials( const FString& UserName, const FString& UserPassword ) override
 	{
 	}
diff --git a/Engine/Source/Programs/UnrealFrontend/Private/Commands/DeployCommand.cpp b/Engine/Source/Programs/UnrealFrontend/Private/Commands/DeployCommand.cpp
deleted file mode 100644
diff --git a/Engine/Source/Programs/UnrealFrontend/Private/Commands/DeployCommand.h b/Engine/Source/Programs/UnrealFrontend/Private/Commands/DeployCommand.h
deleted file mode 100644
diff --git a/Engine/Source/Programs/UnrealFrontend/Private/Commands/LaunchCommand.cpp b/Engine/Source/Programs/UnrealFrontend/Private/Commands/LaunchCommand.cpp
deleted file mode 100644
diff --git a/Engine/Source/Programs/UnrealFrontend/Private/Commands/LaunchCommand.h b/Engine/Source/Programs/UnrealFrontend/Private/Commands/LaunchCommand.h
deleted file mode 100644
diff --git a/Engine/Source/Programs/UnrealFrontend/Private/Commands/PackageCommand.cpp b/Engine/Source/Programs/UnrealFrontend/Private/Commands/PackageCommand.cpp
deleted file mode 100644
diff --git a/Engine/Source/Programs/UnrealFrontend/Private/Commands/PackageCommand.h b/Engine/Source/Programs/UnrealFrontend/Private/Commands/PackageCommand.h
deleted file mode 100644
diff --git a/Engine/Source/Programs/UnrealFrontend/Private/UnrealFrontendMain.cpp b/Engine/Source/Programs/UnrealFrontend/Private/UnrealFrontendMain.cpp
--- a/Engine/Source/Programs/UnrealFrontend/Private/UnrealFrontendMain.cpp
+++ b/Engine/Source/Programs/UnrealFrontend/Private/UnrealFrontendMain.cpp
@@ -3,9 +3,6 @@
 #include "UnrealFrontendMain.h"
 #include "RequiredProgramMainCPPInclude.h"
 
-#include "DeployCommand.h"
-#include "LaunchCommand.h"
-#include "PackageCommand.h"
 #include "StatsConvertCommand.h"
 #include "StatsDumpMemoryCommand.h"
 #include "UserInterfaceCommand.h"
@@ -73,19 +70,7 @@ int32 UnrealFrontendMain( const TCHAR* CommandLine )
 	// NOTE - Some commands may require extra command line parameters.
 	if (bRunCommand)
 	{
-		if (Command.Equals(TEXT("PACKAGE"), ESearchCase::IgnoreCase))
-		{
-			FPackageCommand::Run();
-		}
-		else if (Command.Equals(TEXT("DEPLOY"), ESearchCase::IgnoreCase))
-		{
-			Succeeded = FDeployCommand::Run();
-		}
-		else if (Command.Equals(TEXT("LAUNCH"), ESearchCase::IgnoreCase))
-		{
-			Succeeded = FLaunchCommand::Run(Params);
-		}
-		else if (Command.Equals(TEXT("CONVERT"), ESearchCase::IgnoreCase))
+		if (Command.Equals(TEXT("CONVERT"), ESearchCase::IgnoreCase))
 		{
 			FStatsConvertCommand::Run();
 		}
