Fixed: Editor crash on canceling Launch On Device

diff --git a/Engine/Source/Developer/TargetDeviceServices/Private/Proxies/TargetDeviceProxy.cpp b/Engine/Source/Developer/TargetDeviceServices/Private/Proxies/TargetDeviceProxy.cpp
--- a/Engine/Source/Developer/TargetDeviceServices/Private/Proxies/TargetDeviceProxy.cpp
+++ b/Engine/Source/Developer/TargetDeviceServices/Private/Proxies/TargetDeviceProxy.cpp
@@ -291,5 +291,5 @@ void FTargetDeviceProxy::Reboot()
 
 void FTargetDeviceProxy::InitializeMessaging()
 {
-	// could delete this function
+	MessageEndpoint = FMessageEndpoint::Builder(FName(*FString::Printf(TEXT("FTargetDeviceProxy (%s)"), *Name))).Build();
 }
