added LLM_SCOPE_BYTAG to tag allocations

https://p4-swarm.epicgames.net/reviews/24503827

diff --git a/Engine/Source/Developer/TargetDeviceServices/Private/Proxies/TargetDeviceProxyManager.cpp b/Engine/Source/Developer/TargetDeviceServices/Private/Proxies/TargetDeviceProxyManager.cpp
--- a/Engine/Source/Developer/TargetDeviceServices/Private/Proxies/TargetDeviceProxyManager.cpp
+++ b/Engine/Source/Developer/TargetDeviceServices/Private/Proxies/TargetDeviceProxyManager.cpp
@@ -13,6 +13,7 @@
 /** Defines the interval in seconds in which devices are being pinged by the proxy manager. */
 #define TARGET_DEVICE_SERVICES_PING_INTERVAL 2.5f
 
+LLM_DEFINE_TAG(TargetDeviceProxyManager);
 
 /* FTargetDeviceProxyManager structors
  *****************************************************************************/
@@ -165,6 +166,8 @@ void FTargetDeviceProxyManager::HandlePongMessage(const FTargetDeviceServicePong
  		return;
  	}
 
+	LLM_SCOPE_BYTAG(TargetDeviceProxyManager);
+
 	AddProxyFromPongMessage(Message, Context, false);
 
 	if (Message.Aggregated)
