Add more LLM tags to platform-related editor things.

#rnx

diff --git a/Engine/Source/Developer/TargetDeviceServices/Private/Proxies/TargetDeviceProxyManager.cpp b/Engine/Source/Developer/TargetDeviceServices/Private/Proxies/TargetDeviceProxyManager.cpp
--- a/Engine/Source/Developer/TargetDeviceServices/Private/Proxies/TargetDeviceProxyManager.cpp
+++ b/Engine/Source/Developer/TargetDeviceServices/Private/Proxies/TargetDeviceProxyManager.cpp
@@ -198,6 +198,8 @@ void FTargetDeviceProxyManager::AddProxyFromPongMessage(const FTargetDeviceServi
 
 bool FTargetDeviceProxyManager::HandleTicker(float DeltaTime)
 {
+	LLM_SCOPE_BYTAG(TargetDeviceProxyManager);
+
 	RemoveDeadProxies();
 	SendPing();
 
diff --git a/Engine/Source/Developer/TargetDeviceServices/Private/Services/TargetDeviceService.cpp b/Engine/Source/Developer/TargetDeviceServices/Private/Services/TargetDeviceService.cpp
--- a/Engine/Source/Developer/TargetDeviceServices/Private/Services/TargetDeviceService.cpp
+++ b/Engine/Source/Developer/TargetDeviceServices/Private/Services/TargetDeviceService.cpp
@@ -4,6 +4,7 @@
 
 #include "HAL/PlatformProcess.h"
 #include "HAL/FileManager.h"
+#include "HAL/LowLevelMemTracker.h"
 #include "IMessageBus.h"
 #include "Interfaces/ITargetDevice.h"
 #include "Interfaces/ITargetPlatform.h"
@@ -16,6 +17,8 @@
 
 #include "TargetDeviceServiceMessages.h"
 
+LLM_DECLARE_TAG(TargetDeviceProxyManager);
+
 
 /* Local helpers
  *****************************************************************************/
@@ -406,6 +409,8 @@ void FTargetDeviceService::HandlePingMessage(const FTargetDeviceServicePing& InM
 		return;
 	}
 
+	LLM_SCOPE_BYTAG(TargetDeviceProxyManager);
+
 	ITargetDevicePtr DefaultDevice = GetDevice(); // Default Device is needed here!
 
 	if (DefaultDevice.IsValid())
diff --git a/Engine/Source/Editor/TurnkeySupport/Private/TurnkeySupportModule.cpp b/Engine/Source/Editor/TurnkeySupport/Private/TurnkeySupportModule.cpp
--- a/Engine/Source/Editor/TurnkeySupport/Private/TurnkeySupportModule.cpp
+++ b/Engine/Source/Editor/TurnkeySupport/Private/TurnkeySupportModule.cpp
@@ -2457,6 +2457,7 @@ void FTurnkeySupportModule::UpdateSdkInfoForDevices(TArray<FString> PlatformDevi
 
 		AsyncTask(ENamedThreads::GameThread, [this, ReportFilename, PlatformDeviceIds, ExitCode, TurnkeyProcess]()
 		{
+			LLM_SCOPE(ELLMTag::EngineMisc);
 			FScopeLock Lock(&GTurnkeySection);
 
 			if (!IsEngineExitRequested() && (ExitCode == 0 || ExitCode == 10))
