TargetDeviceServices - Add ability to work with information stored in the TargetDeviceServices module via blueprints and Python.

diff --git a/Engine/Source/Developer/TargetDeviceServices/Private/TargetDeviceServicesBlueprintFunctionLibrary.cpp b/Engine/Source/Developer/TargetDeviceServices/Private/TargetDeviceServicesBlueprintFunctionLibrary.cpp
new file mode 100644
--- /dev/null
+++ b/Engine/Source/Developer/TargetDeviceServices/Private/TargetDeviceServicesBlueprintFunctionLibrary.cpp
@@ -0,0 +1,76 @@
+// Copyright Epic Games, Inc. All Rights Reserved.
+
+#include "TargetDeviceServicesBlueprintFunctionLibrary.h"
+#include "Services/TargetDeviceService.h"
+#include "Modules/ModuleManager.h"
+#include "ITargetDeviceServicesModule.h"
+#include "ITargetDeviceProxy.h"
+
+UTargetDeviceServicesBlueprintFunctionLibrary::UTargetDeviceServicesBlueprintFunctionLibrary(const FObjectInitializer& ObjectInitializer)
+	: Super(ObjectInitializer)
+{
+}
+
+TMap<FString, FDeviceSnapshots> UTargetDeviceServicesBlueprintFunctionLibrary::GetDeviceSnapshots()
+{
+	ITargetDeviceServicesModule& TargetDeviceServicesModule = FModuleManager::GetModuleChecked<ITargetDeviceServicesModule>("TargetDeviceServices");
+	TSharedRef<ITargetDeviceProxyManager> DeviceProxyManager = TargetDeviceServicesModule.GetDeviceProxyManager();
+	TSharedRef<ITargetDeviceServiceManager> DeviceServiceManager = TargetDeviceServicesModule.GetDeviceServiceManager();
+
+	TArray<ITargetDeviceServicePtr> DeviceServices;
+	const int32 DeviceServicesCount = DeviceServiceManager->GetServices(DeviceServices);
+
+	TMap<FString, FDeviceSnapshots> Snapshots;
+	for (int32 CurrentDeviceServiceIndex = 0; CurrentDeviceServiceIndex < DeviceServicesCount; ++CurrentDeviceServiceIndex)
+	{
+		const ITargetDeviceServicePtr& DeviceService = DeviceServices[CurrentDeviceServiceIndex];
+
+		if (DeviceService.IsValid())
+		{
+			// Fetching shared pointer to default device by passing into the ITargetDeviceService::GetDevice method Name_None
+			ITargetDevicePtr DefaultDevice = DeviceService->GetDevice(NAME_None);
+
+			if (DefaultDevice.IsValid())
+			{
+				const FString DefaultDeviceType = TargetDeviceTypes::ToString(DefaultDevice->GetDeviceType());
+				FString DefaultDeviceId = DefaultDevice->GetId().ToString();
+				TSharedPtr<ITargetDeviceProxy> DefaultDeviceProxy = 
+					DeviceProxyManager->FindProxyDeviceForTargetDevice(DefaultDeviceId);
+
+				FDeviceSnapshot DeviceSnapshot{
+					// Fills the Name field
+					DeviceService->GetDeviceName(),
+					// Fills the HostName field
+					DefaultDeviceProxy.IsValid() ? DefaultDeviceProxy->GetHostName() : DefaultDevice->GetName(),
+					// Fills the DeviceType field
+					DefaultDeviceType,
+					// fills the ModelId field
+					DefaultDevice->GetModelId(),
+					// fills the DeviceConnectionType;
+					DefaultDeviceProxy.IsValid() ? DefaultDeviceProxy->GetConnectionType() : TEXT(""),
+					// fills the DeviceId field
+					MoveTemp(DefaultDeviceId),
+					// fills the OperatingSystem field
+					DeviceService->GetDevicePlatformDisplayName(),
+					// fills the IsConnected field
+					DefaultDevice->IsConnected()
+				};
+
+				FDeviceSnapshots* DeviceSnapshotsOfCurrentDeviceTypePtr = Snapshots.Find(DefaultDeviceType);
+				if (nullptr == DeviceSnapshotsOfCurrentDeviceTypePtr)
+				{
+					// Initializer list constructor call
+					TArray<FDeviceSnapshot> Entries{MoveTemp(DeviceSnapshot)};
+					Snapshots.Emplace(DefaultDeviceType, MoveTemp(Entries));
+				}
+				else
+				{
+					DeviceSnapshotsOfCurrentDeviceTypePtr->Entries.Emplace(MoveTemp(DeviceSnapshot));
+				}
+			}
+		}
+	}
+
+	return MoveTemp(Snapshots);
+}
+
diff --git a/Engine/Source/Developer/TargetDeviceServices/Public/TargetDeviceServicesBlueprintFunctionLibrary.h b/Engine/Source/Developer/TargetDeviceServices/Public/TargetDeviceServicesBlueprintFunctionLibrary.h
new file mode 100644
--- /dev/null
+++ b/Engine/Source/Developer/TargetDeviceServices/Public/TargetDeviceServicesBlueprintFunctionLibrary.h
@@ -0,0 +1,89 @@
+// Copyright Epic Games, Inc. All Rights Reserved.
+#pragma once
+
+#include "Kismet/BlueprintFunctionLibrary.h"
+#include "CoreMinimal.h"
+#include "Containers/Map.h"
+#include "TargetDeviceServicesBlueprintFunctionLibrary.generated.h"
+
+/**
+* The struct is designed to store device information
+*/
+USTRUCT(BlueprintType)
+struct FDeviceSnapshot
+{
+	GENERATED_BODY()
+
+public:
+
+	/** Stores device's name. */
+	UPROPERTY(BlueprintReadOnly, Category = "TargetDeviceServices")
+	FString Name;
+
+	/** Stores device's hostname. */
+	UPROPERTY(BlueprintReadOnly, Category = "TargetDeviceServices")
+	FString HostName;
+
+	/** Stores device's type. */
+	UPROPERTY(BlueprintReadOnly, Category = "TargetDeviceServices")
+	FString DeviceType;
+
+	/** Stores device's model identifier. */
+	UPROPERTY(BlueprintReadOnly, Category = "TargetDeviceServices")
+	FString ModelId;
+
+	/** Stores device's connection type. */
+	UPROPERTY(BlueprintReadOnly, Category = "TargetDeviceServices")
+	FString DeviceConnectionType;
+
+	/** Stores device's identifier. */
+	UPROPERTY(BlueprintReadOnly, Category = "TargetDeviceServices")
+	FString DeviceId;
+	
+	/** Stores device's operating system name. */
+	UPROPERTY(BlueprintReadOnly, Category = "TargetDeviceServices")
+	FString OperatingSystem;
+	
+	/**
+	* Stores device's flag that is used to detect whether device
+	* is connected (true) or disconnected (false). 
+	*/
+	UPROPERTY(BlueprintReadOnly, Category = "TargetDeviceServices")
+	bool IsConnected;
+};
+
+/**
+* The struct is a container class that stores instances of FDeviceSnapshot instances.
+* The struct is a helper class to be used in blueprints as a value of TMap container.
+*/
+USTRUCT(BlueprintType)
+struct FDeviceSnapshots
+{
+	GENERATED_BODY()
+
+public:
+
+	/** Stores array of device snapshots. */
+	UPROPERTY(BlueprintReadOnly, Category = "TargetDeviceServices")
+	TArray<FDeviceSnapshot> Entries;
+};
+
+
+/**
+* The class declares a set of functions to be exposed to blueprints.
+*/
+UCLASS(meta = (ScriptName = "TargetDeviceServices"))
+class UTargetDeviceServicesBlueprintFunctionLibrary : public UBlueprintFunctionLibrary
+{
+	GENERATED_UCLASS_BODY()
+	
+public:
+	/**
+	* Fetches snapshots of devices that are available in the network.
+	* 
+	* @return A dictionary of devices' informational snapshots that are grouped by device type (device type string is used as a key).
+	*/
+	UFUNCTION(BlueprintCallable, Category = "TargetDeviceServices")
+	static TARGETDEVICESERVICES_API TMap<FString, FDeviceSnapshots> GetDeviceSnapshots();
+};
+
diff --git a/Engine/Source/Developer/TargetDeviceServices/TargetDeviceServices.Build.cs b/Engine/Source/Developer/TargetDeviceServices/TargetDeviceServices.Build.cs
--- a/Engine/Source/Developer/TargetDeviceServices/TargetDeviceServices.Build.cs
+++ b/Engine/Source/Developer/TargetDeviceServices/TargetDeviceServices.Build.cs
@@ -11,6 +11,7 @@ namespace UnrealBuildTool.Rules
 					"Core",
 					"TargetPlatform",
 					"DesktopPlatform",
+					"Engine",
 				});
 
 			PrivateDependencyModuleNames.AddRange(
