TargetDeviceSessions - Fix build issue for Mac Editor

diff --git a/Engine/Source/Developer/TargetDeviceServices/Private/TargetDeviceServicesBlueprintFunctionLibrary.cpp b/Engine/Source/Developer/TargetDeviceServices/Private/TargetDeviceServicesBlueprintFunctionLibrary.cpp
--- a/Engine/Source/Developer/TargetDeviceServices/Private/TargetDeviceServicesBlueprintFunctionLibrary.cpp
+++ b/Engine/Source/Developer/TargetDeviceServices/Private/TargetDeviceServicesBlueprintFunctionLibrary.cpp
@@ -59,9 +59,7 @@ TMap<FString, FDeviceSnapshots> UTargetDeviceServicesBlueprintFunctionLibrary::G
 				FDeviceSnapshots* DeviceSnapshotsOfCurrentDeviceTypePtr = Snapshots.Find(DefaultDeviceType);
 				if (nullptr == DeviceSnapshotsOfCurrentDeviceTypePtr)
 				{
-					// Initializer list constructor call
-					TArray<FDeviceSnapshot> Entries{MoveTemp(DeviceSnapshot)};
-					Snapshots.Emplace(DefaultDeviceType, MoveTemp(Entries));
+					Snapshots.Emplace(DefaultDeviceType, FDeviceSnapshots{TArray<FDeviceSnapshot>{MoveTemp(DeviceSnapshot)}});
 				}
 				else
 				{
