[Backout] - CL29092845
#rnx
Original CL Desc
-----------------------------------------------------------------
TargetDeviceSessions - Fix build issue for Mac Editor

diff --git a/Engine/Source/Developer/TargetDeviceServices/Private/TargetDeviceServicesBlueprintFunctionLibrary.cpp b/Engine/Source/Developer/TargetDeviceServices/Private/TargetDeviceServicesBlueprintFunctionLibrary.cpp
--- a/Engine/Source/Developer/TargetDeviceServices/Private/TargetDeviceServicesBlueprintFunctionLibrary.cpp
+++ b/Engine/Source/Developer/TargetDeviceServices/Private/TargetDeviceServicesBlueprintFunctionLibrary.cpp
@@ -59,7 +59,9 @@ TMap<FString, FDeviceSnapshots> UTargetDeviceServicesBlueprintFunctionLibrary::G
 				FDeviceSnapshots* DeviceSnapshotsOfCurrentDeviceTypePtr = Snapshots.Find(DefaultDeviceType);
 				if (nullptr == DeviceSnapshotsOfCurrentDeviceTypePtr)
 				{
-					Snapshots.Emplace(DefaultDeviceType, FDeviceSnapshots{TArray<FDeviceSnapshot>{MoveTemp(DeviceSnapshot)}});
+					// Initializer list constructor call
+					TArray<FDeviceSnapshot> Entries{MoveTemp(DeviceSnapshot)};
+					Snapshots.Emplace(DefaultDeviceType, MoveTemp(Entries));
 				}
 				else
 				{
