[Backout] - CL29089501
Original CL Desc
-----------------------------------------------------------------
TargetDeviceServices - Add ability to work with information stored in the TargetDeviceServices module via blueprints and Python.

diff --git a/Engine/Source/Developer/TargetDeviceServices/Private/TargetDeviceServicesBlueprintFunctionLibrary.cpp b/Engine/Source/Developer/TargetDeviceServices/Private/TargetDeviceServicesBlueprintFunctionLibrary.cpp
deleted file mode 100644
diff --git a/Engine/Source/Developer/TargetDeviceServices/Public/TargetDeviceServicesBlueprintFunctionLibrary.h b/Engine/Source/Developer/TargetDeviceServices/Public/TargetDeviceServicesBlueprintFunctionLibrary.h
deleted file mode 100644
diff --git a/Engine/Source/Developer/TargetDeviceServices/TargetDeviceServices.Build.cs b/Engine/Source/Developer/TargetDeviceServices/TargetDeviceServices.Build.cs
--- a/Engine/Source/Developer/TargetDeviceServices/TargetDeviceServices.Build.cs
+++ b/Engine/Source/Developer/TargetDeviceServices/TargetDeviceServices.Build.cs
@@ -11,7 +11,6 @@ namespace UnrealBuildTool.Rules
 					"Core",
 					"TargetPlatform",
 					"DesktopPlatform",
-					"Engine",
 				});
 
 			PrivateDependencyModuleNames.AddRange(
