Fix simple cases of unreachable code for loops that terminate after one iteration

#rnx

diff --git a/Engine/Plugins/Animation/ControlRig/Source/ControlRig/Private/Constraints/ControlRigTransformableHandle.cpp b/Engine/Plugins/Animation/ControlRig/Source/ControlRig/Private/Constraints/ControlRigTransformableHandle.cpp
--- a/Engine/Plugins/Animation/ControlRig/Source/ControlRig/Private/Constraints/ControlRigTransformableHandle.cpp
+++ b/Engine/Plugins/Animation/ControlRig/Source/ControlRig/Private/Constraints/ControlRigTransformableHandle.cpp
@@ -645,8 +645,10 @@ void UTransformableControlHandle::ResolveBoundObjects(FMovieSceneSequenceID Loca
 			return;
 		}
 		
-		for (const TWeakObjectPtr<> ParentObject : ConstraintBindingID.ResolveBoundObjects(LocalSequenceID, SharedPlaybackState))
+		//just do one
+		if (auto BoundObjectsView = ConstraintBindingID.ResolveBoundObjects(LocalSequenceID, SharedPlaybackState); BoundObjectsView.Num() > 0)
 		{
+			const TWeakObjectPtr<> ParentObject = BoundObjectsView[0];
 			const UObject* Bindable = FControlRigObjectBinding::GetBindableObject(ParentObject.Get());
 			if (InControlRig->GetObjectBinding() && InControlRig->GetObjectBinding()->GetBoundObject() == Bindable)
 			{
@@ -654,7 +656,6 @@ void UTransformableControlHandle::ResolveBoundObjects(FMovieSceneSequenceID Loca
 				ControlRig = const_cast<UControlRig*>(InControlRig);
 				RegisterDelegates();
 			}
-			break; //just do one
 		}
 	}
 }
diff --git a/Engine/Plugins/Animation/ControlRig/Source/ControlRigEditor/Private/ControlRigThumbnailRenderer.cpp b/Engine/Plugins/Animation/ControlRig/Source/ControlRigEditor/Private/ControlRigThumbnailRenderer.cpp
--- a/Engine/Plugins/Animation/ControlRig/Source/ControlRigEditor/Private/ControlRigThumbnailRenderer.cpp
+++ b/Engine/Plugins/Animation/ControlRig/Source/ControlRigEditor/Private/ControlRigThumbnailRenderer.cpp
@@ -108,10 +108,9 @@ void UControlRigThumbnailRenderer::AddAdditionalPreviewSceneContent(UObject* Obj
 
 		TArray<UObject*> ArchetypeInstances;
 		CDO->GetArchetypeInstances(ArchetypeInstances);
-		for (UObject* ArchetypeInstance : ArchetypeInstances)
+		if (ArchetypeInstances.Num() > 0)
 		{
-			ControlRig = Cast<UControlRig>(ArchetypeInstance);
-			break;
+			ControlRig = Cast<UControlRig>(ArchetypeInstances[0]);
 		}
 
 		if (ControlRig == nullptr)
diff --git a/Engine/Plugins/Animation/ControlRig/Source/ControlRigEditor/Private/Tools/ControlRigSnapper.cpp b/Engine/Plugins/Animation/ControlRig/Source/ControlRigEditor/Private/Tools/ControlRigSnapper.cpp
--- a/Engine/Plugins/Animation/ControlRig/Source/ControlRigEditor/Private/Tools/ControlRigSnapper.cpp
+++ b/Engine/Plugins/Animation/ControlRig/Source/ControlRigEditor/Private/Tools/ControlRigSnapper.cpp
@@ -358,8 +358,9 @@ static bool CalculateWorldTransformsFromParents(ISequencer* Sequencer, const FCo
 
 				//just do first for now but may do average later.
 				FControlRigSnapper Snapper;
-				for (const FName& Name : ControlRigAndSelection.ControlNames)
+				if (ControlRigAndSelection.ControlNames.Num() > 0)
 				{
+					const FName& Name = ControlRigAndSelection.ControlNames[0];
 					Snapper.GetControlRigControlTransforms(Sequencer, ControlRig, Name, Frames, ParentTransforms, OutParentWorldTransforms);
 					return true;
 				}
diff --git a/Engine/Plugins/ChaosCloth/Source/ChaosCloth/Private/ChaosCloth/ChaosClothingSimulationSolver.cpp b/Engine/Plugins/ChaosCloth/Source/ChaosCloth/Private/ChaosCloth/ChaosClothingSimulationSolver.cpp
--- a/Engine/Plugins/ChaosCloth/Source/ChaosCloth/Private/ChaosCloth/ChaosClothingSimulationSolver.cpp
+++ b/Engine/Plugins/ChaosCloth/Source/ChaosCloth/Private/ChaosCloth/ChaosClothingSimulationSolver.cpp
@@ -1566,8 +1566,9 @@ const Softs::FVelocityAndPressureField& FClothingSimulationSolver::GetWindVeloci
 	{
 		// Return field from first active softbody in this group
 		const TSet<int32>& ParticleRanges = Evolution->GetGroupActiveSoftBodies(GroupId);
-		for (const int32 ParticleRangeId : ParticleRanges)
+		if (auto It = ParticleRanges.CreateConstIterator(); It)
 		{
+			const int32 ParticleRangeId = *It;
 			const TSharedPtr<Softs::FVelocityAndPressureField>& VelocityField = GetClothConstraints(ParticleRangeId).GetVelocityAndPressureField();
 			check(VelocityField);
 			return *VelocityField.Get();
diff --git a/Engine/Plugins/Editor/AssetManagerEditor/Source/AssetManagerEditor/Private/AssetManagerEditorModule.cpp b/Engine/Plugins/Editor/AssetManagerEditor/Source/AssetManagerEditor/Private/AssetManagerEditorModule.cpp
--- a/Engine/Plugins/Editor/AssetManagerEditor/Source/AssetManagerEditor/Private/AssetManagerEditorModule.cpp
+++ b/Engine/Plugins/Editor/AssetManagerEditor/Source/AssetManagerEditor/Private/AssetManagerEditorModule.cpp
@@ -2495,10 +2495,9 @@ void FAssetManagerEditorModule::GetPackageDependenciesPerClass(FName SourcePacka
 	TArray<FAssetData> PackageAssets;
 	if (AssetRegistry->GetAssets(Filter, PackageAssets))
 	{
-		for (const FAssetData& AssetData : PackageAssets)
+		if (PackageAssets.Num() > 0)
 		{
 			OutDependentPackages.AddUnique(SourcePackage);
-			break;
 		}
 	}
 }
diff --git a/Engine/Plugins/Editor/GeometryMode/Source/GeometryMode/Private/GeometryModifiers.cpp b/Engine/Plugins/Editor/GeometryMode/Source/GeometryMode/Private/GeometryModifiers.cpp
--- a/Engine/Plugins/Editor/GeometryMode/Source/GeometryMode/Private/GeometryModifiers.cpp
+++ b/Engine/Plugins/Editor/GeometryMode/Source/GeometryMode/Private/GeometryModifiers.cpp
@@ -2712,10 +2712,9 @@ bool UGeomModifier_Split::OnApply()
 	// Get a pointer to the selected geom object
 
 	FGeomObjectPtr GeomObject = NULL;
-	for( FEdModeGeometry::TGeomObjectIterator Itor( mode->GeomObjectItor() ) ; Itor ; ++Itor )
+	if (FEdModeGeometry::TGeomObjectIterator Itor( mode->GeomObjectItor() ) ; Itor)
 	{
 		GeomObject = *Itor;
-		break;
 	}
 
 	if( !GeomObject.IsValid() )
diff --git a/Engine/Plugins/Editor/ObjectMixer/ObjectMixer/Source/ObjectMixer/Private/Views/List/ObjectMixerEditorList.cpp b/Engine/Plugins/Editor/ObjectMixer/ObjectMixer/Source/ObjectMixer/Private/Views/List/ObjectMixerEditorList.cpp
--- a/Engine/Plugins/Editor/ObjectMixer/ObjectMixer/Source/ObjectMixer/Private/Views/List/ObjectMixerEditorList.cpp
+++ b/Engine/Plugins/Editor/ObjectMixer/ObjectMixer/Source/ObjectMixer/Private/Views/List/ObjectMixerEditorList.cpp
@@ -167,8 +167,9 @@ bool FObjectMixerEditorList::RequestAddObjectsToCollection(const FName& Collecti
 {
 	if (UObjectMixerEditorSerializedData* SerializedData = GetSerializedData())
 	{
-		for (const TSubclassOf<UObjectMixerObjectFilter>& Class : GetObjectFilterClasses())
+		if (auto It = GetObjectFilterClasses().CreateConstIterator(); It)
 		{
+			const TSubclassOf<UObjectMixerObjectFilter>& Class = *It;
 			const FName FilterName = Class->GetFName();
 			if (SerializedData->AddObjectsToCollection(FilterName, CollectionName, ObjectsToAdd))
 			{
diff --git a/Engine/Plugins/Editor/PluginBrowser/Source/PluginBrowser/Private/SPluginCategoryTree.cpp b/Engine/Plugins/Editor/PluginBrowser/Source/PluginBrowser/Private/SPluginCategoryTree.cpp
--- a/Engine/Plugins/Editor/PluginBrowser/Source/PluginBrowser/Private/SPluginCategoryTree.cpp
+++ b/Engine/Plugins/Editor/PluginBrowser/Source/PluginBrowser/Private/SPluginCategoryTree.cpp
@@ -71,13 +71,13 @@ void SPluginCategoryTree::RebuildAndFilterCategoryTree()
 {
 	// Get the path to the first currently selected category
 	TArray<FString> SelectCategoryPath;
-	for(TSharedPtr<FPluginCategory> SelectedItem: TreeView->GetSelectedItems())
+	if (auto SelectedItems = TreeView->GetSelectedItems(); SelectedItems.Num() > 0)
 	{
+		TSharedPtr<FPluginCategory> SelectedItem = SelectedItems[0];
 		for (const FPluginCategory* Category = SelectedItem.Get(); Category != nullptr; Category = Category->ParentCategory.Pin().Get())
 		{
 			SelectCategoryPath.Insert(Category->Name, 0);
 		}
-		break;
 	}
 
 	// Clear the list of plugins in each current category
diff --git a/Engine/Plugins/Enterprise/DatasmithImporter/Source/DatasmithNativeTranslator/Private/DatasmithNativeTranslator.cpp b/Engine/Plugins/Enterprise/DatasmithImporter/Source/DatasmithNativeTranslator/Private/DatasmithNativeTranslator.cpp
--- a/Engine/Plugins/Enterprise/DatasmithImporter/Source/DatasmithNativeTranslator/Private/DatasmithNativeTranslator.cpp
+++ b/Engine/Plugins/Enterprise/DatasmithImporter/Source/DatasmithNativeTranslator/Private/DatasmithNativeTranslator.cpp
@@ -131,10 +131,10 @@ bool FDatasmithNativeTranslator::LoadStaticMesh(const TSharedRef<IDatasmithMeshE
 	{
 		if (DatasmithMesh.bIsCollisionMesh)
 		{
-			for (FMeshDescription& MeshDescription : DatasmithMesh.SourceModels)
+			if (DatasmithMesh.SourceModels.Num() > 0)
 			{
+				FMeshDescription& MeshDescription = DatasmithMesh.SourceModels[0];
 				OutMeshPayload.CollisionMesh = MoveTemp(MeshDescription);
-				break;
 			}
 		}
 		else
diff --git a/Engine/Plugins/Experimental/AnimNext/Source/AnimNextEditor/Private/Param/ObjectPropertyLocatorEditor.cpp b/Engine/Plugins/Experimental/AnimNext/Source/AnimNextEditor/Private/Param/ObjectPropertyLocatorEditor.cpp
--- a/Engine/Plugins/Experimental/AnimNext/Source/AnimNextEditor/Private/Param/ObjectPropertyLocatorEditor.cpp
+++ b/Engine/Plugins/Experimental/AnimNext/Source/AnimNextEditor/Private/Param/ObjectPropertyLocatorEditor.cpp
@@ -86,7 +86,7 @@ class SObjectPropertyLocatorEditor : public SCompoundWidget
 		virtual bool CanExpandScriptStruct(const FStructProperty* StructProperty) const override
 		{
 			// Expand structs that have object properties
-			for (TFieldIterator<FObjectProperty> PropertyIt(StructProperty->Struct, EFieldIteratorFlags::IncludeSuper); PropertyIt; ++PropertyIt)
+			if (TFieldIterator<FObjectProperty> PropertyIt(StructProperty->Struct, EFieldIteratorFlags::IncludeSuper); PropertyIt)
 			{
 				return true;
 			}
diff --git a/Engine/Plugins/Experimental/ChaosCaching/Source/ChaosCaching/Private/Chaos/ChaosCache.cpp b/Engine/Plugins/Experimental/ChaosCaching/Source/ChaosCaching/Private/Chaos/ChaosCache.cpp
--- a/Engine/Plugins/Experimental/ChaosCaching/Source/ChaosCaching/Private/Chaos/ChaosCache.cpp
+++ b/Engine/Plugins/Experimental/ChaosCaching/Source/ChaosCaching/Private/Chaos/ChaosCache.cpp
@@ -372,15 +372,14 @@ void UChaosCache::FlushPendingFrames()
 		for(const TPair<FName, FRichCurves>& Pair : ChannelsTracks)
 		{
 			const FRichCurves& RichCurves = Pair.Value;
-			for(const FRichCurve& Curve : RichCurves.RichCurves)
+			if (RichCurves.RichCurves.Num() > 0)
 			{
+				// assuming all curves have the same time range
+				const FRichCurve& Curve = RichCurves.RichCurves[0];
 				float CurveMin, CurveMax;
 				Curve.GetTimeRange(CurveMin, CurveMax);
 				Min = FMath::Min(Min, CurveMin);
 				Max = FMath::Max(Max, CurveMax);
-
-				// assuming all curves have the same time range
-				break;
 			}
 		}
 
diff --git a/Engine/Plugins/Experimental/ChaosEditor/Source/FractureEditor/Private/FractureEditorModeToolkit.cpp b/Engine/Plugins/Experimental/ChaosEditor/Source/FractureEditor/Private/FractureEditorModeToolkit.cpp
--- a/Engine/Plugins/Experimental/ChaosEditor/Source/FractureEditor/Private/FractureEditorModeToolkit.cpp
+++ b/Engine/Plugins/Experimental/ChaosEditor/Source/FractureEditor/Private/FractureEditorModeToolkit.cpp
@@ -1925,10 +1925,10 @@ void FFractureEditorModeToolkit::SetOutlinerComponents(const TArray<UGeometryCol
 	// update view rest collection
 	UFractureSettings* FractureSettings = GetMutableDefault<UFractureSettings>();
 	FractureSettings->RestCollection = nullptr;
-	for (UGeometryCollectionComponent* Component : InNewComponents)
+	if (InNewComponents.Num() > 0)
 	{
+		UGeometryCollectionComponent* Component = InNewComponents[0];
 		FractureSettings->RestCollection = Component->GetRestCollection();
-		break;
 	}
 
 	RefreshVariableOverrideDetailView(FractureSettings->RestCollection.Get());
diff --git a/Engine/Plugins/Experimental/ChaosNiagara/Source/ChaosNiagara/Private/NiagaraDataInterfaceChaosDestruction.cpp b/Engine/Plugins/Experimental/ChaosNiagara/Source/ChaosNiagara/Private/NiagaraDataInterfaceChaosDestruction.cpp
--- a/Engine/Plugins/Experimental/ChaosNiagara/Source/ChaosNiagara/Private/NiagaraDataInterfaceChaosDestruction.cpp
+++ b/Engine/Plugins/Experimental/ChaosNiagara/Source/ChaosNiagara/Private/NiagaraDataInterfaceChaosDestruction.cpp
@@ -2492,10 +2492,10 @@ bool UNiagaraDataInterfaceChaosDestruction::PerInstanceTick(void* PerInstanceDat
 	FNDIChaosDestruction_InstanceData* InstData = (FNDIChaosDestruction_InstanceData*)PerInstanceData;
 
 	// Update SolverTime
-	for (FSolverData SolverData : Solvers)
+	if (Solvers.Num() > 0)
 	{
+		FSolverData SolverData = Solvers[0];
 		SolverTime = SolverData.Solver->GetSolverTime();
-		break;
 	}
 
 	ShouldSpawn = false;
diff --git a/Engine/Plugins/FX/Niagara/Source/NiagaraEditor/Private/ViewModels/NiagaraParameterPanelViewModel.cpp b/Engine/Plugins/FX/Niagara/Source/NiagaraEditor/Private/ViewModels/NiagaraParameterPanelViewModel.cpp
--- a/Engine/Plugins/FX/Niagara/Source/NiagaraEditor/Private/ViewModels/NiagaraParameterPanelViewModel.cpp
+++ b/Engine/Plugins/FX/Niagara/Source/NiagaraEditor/Private/ViewModels/NiagaraParameterPanelViewModel.cpp
@@ -217,12 +217,12 @@ bool INiagaraImmutableParameterPanelViewModel::GetCanCopyParameterReferenceAndTo
 
 void INiagaraImmutableParameterPanelViewModel::CopyParameterMetaData(const FNiagaraParameterPanelItemBase ItemToCopy) const
 {
-	for (const UNiagaraScriptVariable* ScriptVariable : GetEditableScriptVariablesWithName(ItemToCopy.GetVariable().GetName()))
+	if (auto It = GetEditableScriptVariablesWithName(ItemToCopy.GetVariable().GetName()).CreateConstIterator(); It)
 	{
+		const UNiagaraScriptVariable* ScriptVariable = *It;
 		UNiagaraClipboardContent* ClipboardContent = UNiagaraClipboardContent::Create();
 		ClipboardContent->ScriptVariables.Add({*ScriptVariable});
 		FNiagaraEditorModule::Get().GetClipboard().SetClipboardContent(ClipboardContent);
-		break;
 	}
 }
 
diff --git a/Engine/Plugins/Importers/USDImporter/Source/USDSchemas/Private/USDVolVolumeTranslator.cpp b/Engine/Plugins/Importers/USDImporter/Source/USDSchemas/Private/USDVolVolumeTranslator.cpp
--- a/Engine/Plugins/Importers/USDImporter/Source/USDSchemas/Private/USDVolVolumeTranslator.cpp
+++ b/Engine/Plugins/Importers/USDImporter/Source/USDSchemas/Private/USDVolVolumeTranslator.cpp
@@ -988,8 +988,9 @@ void FUsdVolVolumeTranslator::UpdateComponents(USceneComponent* SceneComponent)
 			TArray<FGuid> ParameterIds;
 			CurrentMaterial->GetAllSparseVolumeTextureParameterInfo(ParameterInfo, ParameterIds);
 
-			for (const FMaterialParameterInfo& Info : ParameterInfo)
+			if (ParameterInfo.Num() > 0)
 			{
+				const FMaterialParameterInfo& Info = ParameterInfo[0];
 				USparseVolumeTexture* SparseVolumeTexture = nullptr;
 				if (CurrentMaterial->GetSparseVolumeTextureParameterValue(Info, SparseVolumeTexture) && SparseVolumeTexture)
 				{
@@ -1044,8 +1045,6 @@ void FUsdVolVolumeTranslator::UpdateComponents(USceneComponent* SceneComponent)
 						}
 					}
 				}
-
-				break;
 			}
 		}
 	}
diff --git a/Engine/Plugins/Importers/USDImporter/Source/USDStage/Private/USDLevelSequenceHelper.cpp b/Engine/Plugins/Importers/USDImporter/Source/USDStage/Private/USDLevelSequenceHelper.cpp
--- a/Engine/Plugins/Importers/USDImporter/Source/USDStage/Private/USDLevelSequenceHelper.cpp
+++ b/Engine/Plugins/Importers/USDImporter/Source/USDStage/Private/USDLevelSequenceHelper.cpp
@@ -2971,8 +2971,12 @@ void FUsdLevelSequenceHelperImpl::AddVolumeTracks(const UUsdPrimTwin& PrimTwin,
 		TArray<FGuid> ParameterIds;
 		CurrentMaterial->GetAllSparseVolumeTextureParameterInfo(ParameterInfo, ParameterIds);
 
-		for (const FMaterialParameterInfo& Info : ParameterInfo)
+		// Follow the theme of only ever caring about the first SVT parameter of the material, as that is all
+		// that the UHeterogeneousVolumeComponent will ever animate anyway. If we're in here we already know
+		// that this first SVT should be animated at any case
+		if (ParameterInfo.Num() > 0)
 		{
+			const FMaterialParameterInfo& Info = ParameterInfo[0];
 			USparseVolumeTexture* SparseVolumeTexture = nullptr;
 			if (CurrentMaterial->GetSparseVolumeTextureParameterValue(Info, SparseVolumeTexture) && SparseVolumeTexture)
 			{
@@ -2981,11 +2985,6 @@ void FUsdLevelSequenceHelperImpl::AddVolumeTracks(const UUsdPrimTwin& PrimTwin,
 					UserData = Cast<UUsdSparseVolumeTextureAssetUserData>(UsdUnreal::ObjectUtils::GetAssetUserData(SparseVolumeTexture));
 				}
 			}
-
-			// Follow the theme of only ever caring about the first SVT parameter of the material, as that is all
-			// that the UHeterogeneousVolumeComponent will ever animate anyway. If we're in here we already know
-			// that this first SVT should be animated at any case
-			break;
 		}
 	}
 	if (!UserData)
diff --git a/Engine/Plugins/Importers/USDImporter/Source/USDStageEditor/Private/SUSDLayersTreeView.cpp b/Engine/Plugins/Importers/USDImporter/Source/USDStageEditor/Private/SUSDLayersTreeView.cpp
--- a/Engine/Plugins/Importers/USDImporter/Source/USDStageEditor/Private/SUSDLayersTreeView.cpp
+++ b/Engine/Plugins/Importers/USDImporter/Source/USDStageEditor/Private/SUSDLayersTreeView.cpp
@@ -1400,15 +1400,11 @@ bool SUsdLayersTreeView::CanIsolateSelectedLayer() const
 		return false;
 	}
 
-	for (FUsdLayerViewModelRef SelectedItem : MySelectedItems)
-	{
-		const UE::FSdfLayerWeak& Layer = SelectedItem->GetLayer();
-		const bool bLayerIsStageRoot = Layer == UsdStage.GetRootLayer();
-
-		return SelectedItem->IsValid() && !bLayerIsStageRoot;
-	}
+	FUsdLayerViewModelRef SelectedItem = MySelectedItems[0];
+	const UE::FSdfLayerWeak& Layer = SelectedItem->GetLayer();
+	const bool bLayerIsStageRoot = Layer == UsdStage.GetRootLayer();
 
-	return false;
+	return SelectedItem->IsValid() && !bLayerIsStageRoot;
 }
 
 void SUsdLayersTreeView::OnIsolateSelectedLayer()
@@ -1848,11 +1844,10 @@ void SUsdLayersTreeView::OnAddSubLayer()
 	FScopedTransaction Transaction(LOCTEXT("AddExistingSublayerTransaction", "Add existing sublayer"));
 
 	TArray<FUsdLayerViewModelRef> MySelectedItems = GetSelectedItems();
-
-	for (FUsdLayerViewModelRef SelectedItem : MySelectedItems)
+	if (MySelectedItems.Num() > 0)
 	{
+		FUsdLayerViewModelRef SelectedItem = MySelectedItems[0];
 		SelectedItem->AddSubLayer(*SubLayerFile.GetValue());
-		break;
 	}
 
 	RequestTreeRefresh();
@@ -1870,10 +1865,10 @@ void SUsdLayersTreeView::OnNewSubLayer()
 	FScopedTransaction Transaction(LOCTEXT("AddNewSublayerTransaction", "Add new sublayer"));
 
 	TArray<FUsdLayerViewModelRef> MySelectedItems = GetSelectedItems();
-	for (FUsdLayerViewModelRef SelectedItem : MySelectedItems)
+	if (MySelectedItems.Num() > 0)
 	{
+		FUsdLayerViewModelRef SelectedItem = MySelectedItems[0];
 		SelectedItem->NewSubLayer(*SubLayerFile.GetValue());
-		break;
 	}
 
 	RequestTreeRefresh();
diff --git a/Engine/Plugins/MovieScene/LevelSequenceEditor/Source/LevelSequenceEditor/Private/LevelSequenceEditorSubsystem.cpp b/Engine/Plugins/MovieScene/LevelSequenceEditor/Source/LevelSequenceEditor/Private/LevelSequenceEditorSubsystem.cpp
--- a/Engine/Plugins/MovieScene/LevelSequenceEditor/Source/LevelSequenceEditor/Private/LevelSequenceEditorSubsystem.cpp
+++ b/Engine/Plugins/MovieScene/LevelSequenceEditor/Source/LevelSequenceEditor/Private/LevelSequenceEditorSubsystem.cpp
@@ -1936,10 +1936,10 @@ bool ULevelSequenceEditorSubsystem::BakeTransformWithSettings(const TArray<FMovi
 				for (UMovieSceneSection* ConstraintSection : ConstraintTrack->GetAllSections())
 				{
 					FMovieSceneObjectBindingID ConstraintBindingID = (Cast<UMovieScene3DConstraintSection>(ConstraintSection))->GetConstraintBindingID();
-					for (TWeakObjectPtr<> ParentObject : ConstraintBindingID.ResolveBoundObjects(Sequencer->GetFocusedTemplateID(), *Sequencer))
+					if (auto BoundObjectsView = ConstraintBindingID.ResolveBoundObjects(Sequencer->GetFocusedTemplateID(), *Sequencer); BoundObjectsView.Num() > 0)
 					{
+						TWeakObjectPtr<> ParentObject = BoundObjectsView[0];
 						AttachParentActor = Cast<AActor>(ParentObject.Get());
-						break;
 					}
 				}
 
diff --git a/Engine/Plugins/MovieScene/MovieRenderPipeline/Source/MovieRenderPipelineCore/Private/Graph/Nodes/MovieGraphBranchNode.cpp b/Engine/Plugins/MovieScene/MovieRenderPipeline/Source/MovieRenderPipelineCore/Private/Graph/Nodes/MovieGraphBranchNode.cpp
--- a/Engine/Plugins/MovieScene/MovieRenderPipeline/Source/MovieRenderPipelineCore/Private/Graph/Nodes/MovieGraphBranchNode.cpp
+++ b/Engine/Plugins/MovieScene/MovieRenderPipeline/Source/MovieRenderPipelineCore/Private/Graph/Nodes/MovieGraphBranchNode.cpp
@@ -41,13 +41,13 @@ TArray<UMovieGraphPin*> UMovieGraphBranchNode::EvaluatePinsToFollow(FMovieGraphE
 	}
 
 	UMovieGraphPin* OtherPin = nullptr;
-	for (UMovieGraphEdge* Edge : ConditionalPin->Edges)
+	
+	// We only support a single connection to this pin type anyways.
+	if (ConditionalPin->Edges.Num() > 0)
 	{
+		UMovieGraphEdge* Edge = ConditionalPin->Edges[0];
 		constexpr bool bFollowRerouteConnections = true;
 		OtherPin = Edge->GetOtherPin(ConditionalPin, bFollowRerouteConnections);
-
-		// We only support a single connection to this pin type anyways.
-		break;
 	}
 
 	// There may not be a node actually connected. We don't know what to do in this case (as our nodes don't have default values)
diff --git a/Engine/Plugins/PCG/Source/PCG/Private/Graph/PCGGraphCompilerGPU.cpp b/Engine/Plugins/PCG/Source/PCG/Private/Graph/PCGGraphCompilerGPU.cpp
--- a/Engine/Plugins/PCG/Source/PCG/Private/Graph/PCGGraphCompilerGPU.cpp
+++ b/Engine/Plugins/PCG/Source/PCG/Private/Graph/PCGGraphCompilerGPU.cpp
@@ -1668,11 +1668,11 @@ void FPCGGraphCompilerGPU::CreateGPUNodes(FPCGGraphCompiler& InOutCompiler, UPCG
 		ComputeGraphTask.NodeId = ComputeGraphTaskId;
 
 		// All nodes in subset will be from same stack/parent, so assign from any.
-		for (FPCGTaskId GPUTaskId : NodeSubsetToConvertToCFGraph)
+		if (auto It = NodeSubsetToConvertToCFGraph.CreateConstIterator(); It)
 		{
+			FPCGTaskId GPUTaskId = *It;
 			ComputeGraphTask.ParentId = InOutCompiledTasks[GPUTaskId].ParentId;
 			ComputeGraphTask.StackIndex = InOutCompiledTasks[GPUTaskId].StackIndex;
-			break;
 		}
 
 		// Wire in the compute graph task, side by side with the individual GPU tasks, which will be culled below.
diff --git a/Engine/Plugins/PCG/Source/PCGEditor/Private/PCGEditor.cpp b/Engine/Plugins/PCG/Source/PCGEditor/Private/PCGEditor.cpp
--- a/Engine/Plugins/PCG/Source/PCGEditor/Private/PCGEditor.cpp
+++ b/Engine/Plugins/PCG/Source/PCGEditor/Private/PCGEditor.cpp
@@ -1435,8 +1435,9 @@ bool FPCGEditor::CanSelectNamedRerouteUsages() const
 		return false;
 	}
 
-	for (const UObject* Object : GraphEditorWidget->GetSelectedNodes())
+	if (auto It = GraphEditorWidget->GetSelectedNodes().CreateConstIterator(); It)
 	{
+		const UObject* Object = *It;
 		return Object && Object->IsA<UPCGEditorGraphNodeNamedRerouteDeclaration>();
 	}
 
@@ -1492,8 +1493,9 @@ bool FPCGEditor::CanSelectNamedRerouteDeclaration() const
 		return false;
 	}
 
-	for (const UObject* Object : GraphEditorWidget->GetSelectedNodes())
+	if (auto It = GraphEditorWidget->GetSelectedNodes().CreateConstIterator(); It)
 	{
+		const UObject* Object = *It;
 		return Object && Object->IsA<UPCGEditorGraphNodeNamedRerouteUsage>();
 	}
 
diff --git a/Engine/Plugins/Runtime/GameplayAbilities/Source/GameplayAbilities/Private/AbilitySystemDebugHUD.cpp b/Engine/Plugins/Runtime/GameplayAbilities/Source/GameplayAbilities/Private/AbilitySystemDebugHUD.cpp
--- a/Engine/Plugins/Runtime/GameplayAbilities/Source/GameplayAbilities/Private/AbilitySystemDebugHUD.cpp
+++ b/Engine/Plugins/Runtime/GameplayAbilities/Source/GameplayAbilities/Private/AbilitySystemDebugHUD.cpp
@@ -299,10 +299,9 @@ void AAbilitySystemDebugHUD::NotifyExtensionEnableChanged(UWorld* InWorld)
 	}
 
 	AAbilitySystemDebugHUD* HUD = nullptr;
-	for (TActorIterator<AAbilitySystemDebugHUD> It(InWorld); It; ++It)
+	if (TActorIterator<AAbilitySystemDebugHUD> It(InWorld); It)
 	{
 		HUD = *It;
-		break;
 	}
 
 	bool bAnyExtensionEnabled = bEnableBasicHUD;
diff --git a/Engine/Plugins/Runtime/Harmonix/Source/HarmonixMetasound/Private/HarmonixMetasound/Components/MetasoundMusicClockDriver.cpp b/Engine/Plugins/Runtime/Harmonix/Source/HarmonixMetasound/Private/HarmonixMetasound/Components/MetasoundMusicClockDriver.cpp
--- a/Engine/Plugins/Runtime/Harmonix/Source/HarmonixMetasound/Private/HarmonixMetasound/Components/MetasoundMusicClockDriver.cpp
+++ b/Engine/Plugins/Runtime/Harmonix/Source/HarmonixMetasound/Private/HarmonixMetasound/Components/MetasoundMusicClockDriver.cpp
@@ -441,13 +441,12 @@ void FMetasoundMusicClockDriver::OnGeneratorIOUpdatedWithChanges(const TArray<Me
 {
 	if (!MetasoundOutputName.IsNone())
 	{
-		for (auto& VertexInterfaceChange : VertexInterfaceChanges)
+		if (VertexInterfaceChanges.Num() > 0)
 		{
 			SetClockHistory(UMidiClockUpdateSubsystem::GetOrCreateClockHistory(MidiSongPosAnalyzerAddress));
 			SmoothedAudioRenderClockHistoryCursor = ClockHistory->CreateReadCursor();
 			SmoothedPlayerExperienceClockHistoryCursor = ClockHistory->CreateReadCursor();
 			SmoothedVideoRenderClockHistoryCursor = ClockHistory->CreateReadCursor();
-			break;
 		}
 	}
 	
diff --git a/Engine/Plugins/Runtime/Metasound/Source/MetasoundFrontend/Private/MetasoundFrontendDocumentBuilder.cpp b/Engine/Plugins/Runtime/Metasound/Source/MetasoundFrontend/Private/MetasoundFrontendDocumentBuilder.cpp
--- a/Engine/Plugins/Runtime/Metasound/Source/MetasoundFrontend/Private/MetasoundFrontendDocumentBuilder.cpp
+++ b/Engine/Plugins/Runtime/Metasound/Source/MetasoundFrontend/Private/MetasoundFrontendDocumentBuilder.cpp
@@ -505,11 +505,9 @@ namespace Metasound::Frontend
 							const TMap<FGuid, FVector2D>& Locations = Node->Style.Display.Locations;
 							if (!Locations.IsEmpty())
 							{
-								for (const TPair<FGuid, FVector2D>& Location : Locations)
-								{
-									NextLocation = Location.Value - (NumBeforeDefined * DisplayStyle::NodeLayout::DefaultOffsetY);
-									break;
-								}
+								auto It = Locations.CreateConstIterator();
+								const TPair<FGuid, FVector2D>& Location = *It;
+								NextLocation = Location.Value - (NumBeforeDefined * DisplayStyle::NodeLayout::DefaultOffsetY);
 								break;
 							}
 						}
diff --git a/Engine/Plugins/Runtime/Metasound/Source/MetasoundFrontend/Private/MetasoundFrontendTransform.cpp b/Engine/Plugins/Runtime/Metasound/Source/MetasoundFrontend/Private/MetasoundFrontendTransform.cpp
--- a/Engine/Plugins/Runtime/Metasound/Source/MetasoundFrontend/Private/MetasoundFrontendTransform.cpp
+++ b/Engine/Plugins/Runtime/Metasound/Source/MetasoundFrontend/Private/MetasoundFrontendTransform.cpp
@@ -473,11 +473,9 @@ namespace Metasound
 							const TMap<FGuid, FVector2D>& Locations = NodeHandle->GetNodeStyle().Display.Locations;
 							if (!Locations.IsEmpty())
 							{
-								for (const TPair<FGuid, FVector2D>& Location : Locations)
-								{
-									NextLocation = Location.Value - (NumBeforeDefined * DisplayStyle::NodeLayout::DefaultOffsetY);
-									break;
-								}
+								auto It = Locations.CreateConstIterator();
+								const TPair<FGuid, FVector2D>& Location = *It;
+								NextLocation = Location.Value - (NumBeforeDefined * DisplayStyle::NodeLayout::DefaultOffsetY);
 								break;
 							}
 						}
diff --git a/Engine/Plugins/Runtime/Metasound/Source/MetasoundFrontend/Private/NodeTemplates/MetasoundFrontendNodeTemplateInput.cpp b/Engine/Plugins/Runtime/Metasound/Source/MetasoundFrontend/Private/NodeTemplates/MetasoundFrontendNodeTemplateInput.cpp
--- a/Engine/Plugins/Runtime/Metasound/Source/MetasoundFrontend/Private/NodeTemplates/MetasoundFrontendNodeTemplateInput.cpp
+++ b/Engine/Plugins/Runtime/Metasound/Source/MetasoundFrontend/Private/NodeTemplates/MetasoundFrontendNodeTemplateInput.cpp
@@ -69,19 +69,17 @@ namespace Metasound::Frontend
 					{
 						if (!Node->Style.Display.Locations.IsEmpty())
 						{
-							for (const TPair<FGuid, FVector2D>& Pair : Node->Style.Display.Locations)
+							auto It = Node->Style.Display.Locations.CreateConstIterator();
+							const TPair<FGuid, FVector2D>& Pair = *It;
+							const FGuid ConnectedVertexID = ConnectedVertices.Last().VertexID;
+							const uint32 Index = Node->Interface.Inputs.IndexOfByPredicate([&ConnectedVertexID](const FMetasoundFrontendVertex& Input)
 							{
-								const FGuid ConnectedVertexID = ConnectedVertices.Last().VertexID;
-								const uint32 Index = Node->Interface.Inputs.IndexOfByPredicate([&ConnectedVertexID](const FMetasoundFrontendVertex& Input)
-								{
-									return Input.VertexID == ConnectedVertexID;
-								});
-								NewLocation = Pair.Value;
-								NewLocation -= DisplayStyle::NodeLayout::DefaultOffsetX;
-								// Offset Y position based on connected input index to avoid overlapping nodes
-								NewLocation += Index * DisplayStyle::NodeLayout::DefaultOffsetY;
-								break;
-							}
+								return Input.VertexID == ConnectedVertexID;
+							});
+							NewLocation = Pair.Value;
+							NewLocation -= DisplayStyle::NodeLayout::DefaultOffsetX;
+							// Offset Y position based on connected input index to avoid overlapping nodes
+							NewLocation += Index * DisplayStyle::NodeLayout::DefaultOffsetY;
 						}
 					}
 					
diff --git a/Engine/Plugins/Runtime/StateTree/Source/StateTreeEditorModule/Private/SStateTreeDiff.cpp b/Engine/Plugins/Runtime/StateTree/Source/StateTreeEditorModule/Private/SStateTreeDiff.cpp
--- a/Engine/Plugins/Runtime/StateTree/Source/StateTreeEditorModule/Private/SStateTreeDiff.cpp
+++ b/Engine/Plugins/Runtime/StateTree/Source/StateTreeEditorModule/Private/SStateTreeDiff.cpp
@@ -631,8 +631,9 @@ bool SDiffWidget::ShouldHighlightRow(const TUniquePtr<FAsyncDetailViewDiff::Diff
 		return false;
 	}
 
-	for (const FSingleDiffEntry& BindingDiff : StateBindingDiffs)
+	if (StateBindingDiffs.Num() > 0)
 	{
+		const FSingleDiffEntry& BindingDiff = StateBindingDiffs[0];
 		return BindingDiff.BindingPath.IsSubPropertyMatch(PropertySoftPath) || BindingDiff.BindingPath == PropertySoftPath;
 	}
 
diff --git a/Engine/Source/Developer/LogVisualizer/Private/VisualLoggerDatabase.cpp b/Engine/Source/Developer/LogVisualizer/Private/VisualLoggerDatabase.cpp
--- a/Engine/Source/Developer/LogVisualizer/Private/VisualLoggerDatabase.cpp
+++ b/Engine/Source/Developer/LogVisualizer/Private/VisualLoggerDatabase.cpp
@@ -460,7 +460,7 @@ AActor* FVisualLoggerEditorInterface::GetHelperActor(UWorld* InWorld) const
 		return nullptr;
 	}
 
-	for (TActorIterator<AVisualLoggerRenderingActor> It(World); It; ++It)
+	if (TActorIterator<AVisualLoggerRenderingActor> It(World); It)
 	{
 		return *It;
 	}
diff --git a/Engine/Source/Developer/S3Client/Private/S3Client.cpp b/Engine/Source/Developer/S3Client/Private/S3Client.cpp
--- a/Engine/Source/Developer/S3Client/Private/S3Client.cpp
+++ b/Engine/Source/Developer/S3Client/Private/S3Client.cpp
@@ -294,9 +294,9 @@ bool FS3ClientCredentials::IsValid() const
 ////////////////////////////////////////////////////////////////////////////////
 FS3ClientCredentials FS3CredentialsProfileStore::GetDefault() const
 {
-	for (const auto& KV : Credentials)
+	if (auto It = Credentials.CreateConstIterator(); It)
 	{
-		return KV.Value;
+		return It.Value();
 	}
 
 	return FS3ClientCredentials();
diff --git a/Engine/Source/Developer/TargetDeviceServices/Private/Proxies/TargetDeviceProxy.cpp b/Engine/Source/Developer/TargetDeviceServices/Private/Proxies/TargetDeviceProxy.cpp
--- a/Engine/Source/Developer/TargetDeviceServices/Private/Proxies/TargetDeviceProxy.cpp
+++ b/Engine/Source/Developer/TargetDeviceServices/Private/Proxies/TargetDeviceProxy.cpp
@@ -121,7 +121,7 @@ FName FTargetDeviceProxy::GetTargetDeviceVariant(const FString& InDeviceId) cons
 		const FTargetDeviceProxyVariant& Variant = ItVariant.Value();
 
 		// for an aggregate (All_<platform>_devices_on_<host>) proxy we have a list of associated devices
-		for (TSet<FString>::TConstIterator ItDeviceID(Variant.DeviceIDs); ItDeviceID; ++ItDeviceID)
+		if (TSet<FString>::TConstIterator ItDeviceID(Variant.DeviceIDs); ItDeviceID)
 		{
 			// should return the first device
 			// this method is designed for physical devices, not aggregate proxies
@@ -129,7 +129,6 @@ FName FTargetDeviceProxy::GetTargetDeviceVariant(const FString& InDeviceId) cons
 			{
 				return ItVariant.Key();
 			}
-			break;
 		}
 	}
 
@@ -153,12 +152,11 @@ const FString FTargetDeviceProxy::GetTargetDeviceId(FName InVariant) const
 {
 	FString Variant;
 	// for an aggregate (All_<platform>_devices_on_<host>) proxy we have a list of associated devices
-	for (TSet<FString>::TConstIterator ItVariant(TargetDeviceVariants[(InVariant == NAME_None) ? DefaultVariant : InVariant].DeviceIDs); ItVariant; ++ItVariant)
+	if (TSet<FString>::TConstIterator ItVariant(TargetDeviceVariants[(InVariant == NAME_None) ? DefaultVariant : InVariant].DeviceIDs); ItVariant)
 	{
 		// should return the first device
 		// this method is designed for physical devices, not aggregate proxies
 		Variant = *ItVariant;
-		break;
 	}
 	return Variant;
 }
@@ -219,7 +217,7 @@ bool FTargetDeviceProxy::HasDeviceId(const FString& InDeviceId) const
 	{
 		const FTargetDeviceProxyVariant& Variant = ItVariant.Value();
 
-		for (TSet<FString>::TConstIterator ItDeviceID(Variant.DeviceIDs); ItDeviceID; ++ItDeviceID)
+		if (TSet<FString>::TConstIterator ItDeviceID(Variant.DeviceIDs); ItDeviceID)
 		{
 			// should return the first device
 			// this method is designed for physical devices, not aggregate proxies
@@ -227,7 +225,6 @@ bool FTargetDeviceProxy::HasDeviceId(const FString& InDeviceId) const
 			{
 				return true;
 			}
-			break;
 		}
 	}
 	return false;
diff --git a/Engine/Source/Editor/AnimGraph/Private/K2Node_AnimGetter.cpp b/Engine/Source/Editor/AnimGraph/Private/K2Node_AnimGetter.cpp
--- a/Engine/Source/Editor/AnimGraph/Private/K2Node_AnimGetter.cpp
+++ b/Engine/Source/Editor/AnimGraph/Private/K2Node_AnimGetter.cpp
@@ -478,29 +478,24 @@ bool UK2Node_AnimGetter::IsContextValidForSchema(const UEdGraphSchema* Schema) c
 		// Valid in all graphs
 		return true;
 	}
-	
-	for(const FString& Context : Contexts)
-	{
-		UClass* ClassToCheck = nullptr;
-		if(Context == TEXT("CustomBlend"))
-		{
-			ClassToCheck = UAnimationCustomTransitionSchema::StaticClass();
-		}
-
-		if(Context == TEXT("Transition"))
-		{
-			ClassToCheck = UAnimationTransitionSchema::StaticClass();
-		}
 
-		if(Context == TEXT("AnimGraph"))
-		{
-			ClassToCheck = UAnimationGraphSchema::StaticClass();
-		}
+	const FString& Context = Contexts[0];
 
-		return Schema->GetClass() == ClassToCheck;
+	UClass* ClassToCheck = nullptr;
+	if(Context == TEXT("CustomBlend"))
+	{
+		ClassToCheck = UAnimationCustomTransitionSchema::StaticClass();
+	}
+	else if(Context == TEXT("Transition"))
+	{
+		ClassToCheck = UAnimationTransitionSchema::StaticClass();
+	}
+	else if(Context == TEXT("AnimGraph"))
+	{
+		ClassToCheck = UAnimationGraphSchema::StaticClass();
 	}
 
-	return false;
+	return Schema->GetClass() == ClassToCheck;
 }
 
 void UK2Node_AnimGetter::UpdateCachedTitle(FNodeSpawnData& SpawnData)
diff --git a/Engine/Source/Editor/LandscapeEditor/Private/LandscapeEdMode.cpp b/Engine/Source/Editor/LandscapeEditor/Private/LandscapeEdMode.cpp
--- a/Engine/Source/Editor/LandscapeEditor/Private/LandscapeEdMode.cpp
+++ b/Engine/Source/Editor/LandscapeEditor/Private/LandscapeEdMode.cpp
@@ -734,10 +734,9 @@ void FEdModeLandscape::Exit()
 
 	// Clear all GizmoActors if there is no Landscape in World
 	bool bIsLandscapeExist = false;
-	for (TActorIterator<ALandscapeProxy> It(GetWorld()); It; ++It)
+	if (TActorIterator<ALandscapeProxy> It(GetWorld()); It)
 	{
 		bIsLandscapeExist = true;
-		break;
 	}
 
 	if (!bIsLandscapeExist)
@@ -2356,10 +2355,9 @@ int32 FEdModeLandscape::UpdateLandscapeList()
 	if (!CurrentGizmoActor.IsValid())
 	{
 		ALandscapeGizmoActiveActor* GizmoActor = nullptr;
-		for (TActorIterator<ALandscapeGizmoActiveActor> It(GetWorld()); It; ++It)
+		if (TActorIterator<ALandscapeGizmoActiveActor> It(GetWorld()); It)
 		{
 			GizmoActor = *It;
-			break;
 		}
 	}
 
diff --git a/Engine/Source/Editor/LandscapeEditor/Private/LandscapeEdModeMirrorTool.cpp b/Engine/Source/Editor/LandscapeEditor/Private/LandscapeEdModeMirrorTool.cpp
--- a/Engine/Source/Editor/LandscapeEditor/Private/LandscapeEdModeMirrorTool.cpp
+++ b/Engine/Source/Editor/LandscapeEditor/Private/LandscapeEdModeMirrorTool.cpp
@@ -290,8 +290,9 @@ protected:
 		// try to find Z location
 		TSet<ULandscapeComponent*> Components;
 		LandscapeInfo->GetComponentsInRegion(x, y, x, y, Components);
-		for (ULandscapeComponent* Component : Components)
+		if (auto It = Components.CreateConstIterator(); It)
 		{
+			ULandscapeComponent* Component = *It;
 			FLandscapeComponentDataInterface DataInterface(Component);
 			return LandscapeDataAccess::GetLocalHeight(DataInterface.GetHeight(x - Component->SectionBaseX, y - Component->SectionBaseY));
 		}
diff --git a/Engine/Source/Editor/LandscapeEditor/Private/Tests/LandscapeAutomationTests.cpp b/Engine/Source/Editor/LandscapeEditor/Private/Tests/LandscapeAutomationTests.cpp
--- a/Engine/Source/Editor/LandscapeEditor/Private/Tests/LandscapeAutomationTests.cpp
+++ b/Engine/Source/Editor/LandscapeEditor/Private/Tests/LandscapeAutomationTests.cpp
@@ -204,12 +204,12 @@ bool FCheckHeight::Update()
 	}
 
 	
-	for (const auto& InfoMapPair : ULandscapeInfoMap::GetLandscapeInfoMap(World).Map)
+	if (auto It = ULandscapeInfoMap::GetLandscapeInfoMap(World).Map.CreateConstIterator(); It)
 	{
+		const auto& InfoMapPair = *It;
 		TOptional<float> Height = InfoMapPair.Value->GetLandscapeProxy()->GetHeightAtLocation(FVector(0, 0, 0), EHeightfieldSource::Editor);
 		CurrentTest->TestEqual("Has Height Value at 0,0", Height.IsSet(), true);
 		CurrentTest->TestNearlyEqual("Height Value at 0,0 is 0", Height.GetValue(), 0.0f, 1e-4f);
-		break;
 	}
 	return true;
 }
diff --git a/Engine/Source/Editor/LevelEditor/Private/SSurfaceProperties.cpp b/Engine/Source/Editor/LevelEditor/Private/SSurfaceProperties.cpp
--- a/Engine/Source/Editor/LevelEditor/Private/SSurfaceProperties.cpp
+++ b/Engine/Source/Editor/LevelEditor/Private/SSurfaceProperties.cpp
@@ -35,7 +35,7 @@ void SSurfaceProperties::Construct( const FArguments& InArgs )
 	CachedScalingValueV = 1.0f;
 
 	// Initialize scale fields according to the scale of the first selected surface
-	for (TSelectedSurfaceIterator<> It(GetWorld()); It; ++It)
+	if (TSelectedSurfaceIterator<> It(GetWorld()); It)
 	{
 		FBspSurf* Surf = *It;
 		UModel* Model = It.GetModel();
@@ -55,8 +55,6 @@ void SSurfaceProperties::Construct( const FArguments& InArgs )
 		{
 			CachedScalingValueV = 1.0f / TextureVSize;
 		}
-
-		break;
 	}
 
 	bPreserveScaleRatio = false;
diff --git a/Engine/Source/Editor/MovieSceneTools/Private/MovieSceneAlphaBlendOptionCustomization.cpp b/Engine/Source/Editor/MovieSceneTools/Private/MovieSceneAlphaBlendOptionCustomization.cpp
--- a/Engine/Source/Editor/MovieSceneTools/Private/MovieSceneAlphaBlendOptionCustomization.cpp
+++ b/Engine/Source/Editor/MovieSceneTools/Private/MovieSceneAlphaBlendOptionCustomization.cpp
@@ -49,9 +49,9 @@ EAlphaBlendOption FAlphaBlendPropertyCustomization::GetAlphaBlendOption() const
 		TArray<void*> RawData;
 		TypeProperty->AccessRawData(RawData);
 
-		for (void* Ptr : RawData)
+		if (RawData.Num() > 0)
 		{
-			return *((EAlphaBlendOption*)Ptr);
+			return *((EAlphaBlendOption*)RawData[0]);
 		}
 	}
 
diff --git a/Engine/Source/Editor/MovieSceneTools/Private/MovieSceneToolHelpers.cpp b/Engine/Source/Editor/MovieSceneTools/Private/MovieSceneToolHelpers.cpp
--- a/Engine/Source/Editor/MovieSceneTools/Private/MovieSceneToolHelpers.cpp
+++ b/Engine/Source/Editor/MovieSceneTools/Private/MovieSceneToolHelpers.cpp
@@ -5296,13 +5296,13 @@ void MovieSceneToolHelpers::GetActorParentsWithAttachments(ISequencer* Sequencer
 					for (UMovieSceneSection* ConstraintSection : ConstraintTrack->GetAllSections())
 					{
 						FMovieSceneObjectBindingID ConstraintBindingID = (Cast<UMovieScene3DConstraintSection>(ConstraintSection))->GetConstraintBindingID();
-						for (TWeakObjectPtr<> ParentObject : ConstraintBindingID.ResolveBoundObjects(Sequencer->GetFocusedTemplateID(), *Sequencer))
+						if (auto BoundObjectsView = ConstraintBindingID.ResolveBoundObjects(Sequencer->GetFocusedTemplateID(), *Sequencer); BoundObjectsView.Num())
 						{
+							TWeakObjectPtr<> ParentObject = BoundObjectsView[0];
 							FActorForWorldTransforms OutParentActor;
 							OutParentActor.Actor = Cast<AActor>(ParentObject.Get());;
 							OutParentActors.AddUnique(OutParentActor);
 							GetActorParents(OutParentActor, OutParentActors);
-							break;
 						}
 					}
 				}
diff --git a/Engine/Source/Editor/PropertyEditor/Private/PropertyEditorToolkit.cpp b/Engine/Source/Editor/PropertyEditor/Private/PropertyEditorToolkit.cpp
--- a/Engine/Source/Editor/PropertyEditor/Private/PropertyEditorToolkit.cpp
+++ b/Engine/Source/Editor/PropertyEditor/Private/PropertyEditorToolkit.cpp
@@ -322,10 +322,9 @@ void FPropertyEditorToolkit::ToggleColumnForProperty( const TSharedPtr< FPropert
 	TSharedRef< FPropertyPath > NewPath = PropertyPath->TrimRoot( PropertyTable->GetRootPath()->GetNumProperties() );
 	const TSet< TSharedRef< IPropertyTableRow > > SelectedRows = PropertyTable->GetSelectedRows();
 	
-	for( auto RowIter = SelectedRows.CreateConstIterator(); RowIter; ++RowIter )
+	if (auto RowIter = SelectedRows.CreateConstIterator(); RowIter)
 	{
 		NewPath = NewPath->TrimRoot( (*RowIter)->GetPartialPath()->GetNumProperties() );
-		break;
 	}
 
 	if ( NewPath->GetNumProperties() == 0 )
@@ -383,10 +382,9 @@ bool FPropertyEditorToolkit::IsExposedAsColumn( const TWeakPtr< IPropertyTreeRow
 			TSharedRef< FPropertyPath > TrimmedPath = RowPathPtr->TrimRoot( PropertyTable->GetRootPath()->GetNumProperties() );
 			const TSet< TSharedRef< IPropertyTableRow > > SelectedRows = PropertyTable->GetSelectedRows();
 
-			for( auto RowIter = SelectedRows.CreateConstIterator(); RowIter; ++RowIter )
+			if (auto RowIter = SelectedRows.CreateConstIterator(); RowIter)
 			{
 				TrimmedPath = TrimmedPath->TrimRoot( (*RowIter)->GetPartialPath()->GetNumProperties() );
-				break;
 			}
 
 			for (int Index = 0; Index < PropertyPathsAddedAsColumns.Num(); Index++)
@@ -446,10 +444,9 @@ void FPropertyEditorToolkit::GridSelectionChanged()
 
 	if ( SelectedRows.Num() == 1 )
 	{
-		for( auto RowIter = SelectedRows.CreateConstIterator(); RowIter; ++RowIter )
+		if (auto RowIter = SelectedRows.CreateConstIterator(); RowIter)
 		{
 			PropertyTree->SetRootPath( PropertyTable->GetRootPath()->ExtendPath( (*RowIter)->GetPartialPath() ) );
-			break;
 		}
 	}
 	else if ( !FPropertyPath::AreEqual( PropertyTree->GetRootPath(), PropertyTable->GetRootPath() ) )
diff --git a/Engine/Source/Editor/Sequencer/Private/SequencerContextMenus.cpp b/Engine/Source/Editor/Sequencer/Private/SequencerContextMenus.cpp
--- a/Engine/Source/Editor/Sequencer/Private/SequencerContextMenus.cpp
+++ b/Engine/Source/Editor/Sequencer/Private/SequencerContextMenus.cpp
@@ -907,17 +907,14 @@ void FSectionContextMenu::SetSectionToKey()
 	}
 
 	const bool bToggle = IsSectionToKey();
-	for (UMovieSceneSection* Section : Sequencer->GetViewModel()->GetSelection()->GetSelectedSections())
+	auto It = Sequencer->GetViewModel()->GetSelection()->GetSelectedSections().CreateConstIterator();
+	UMovieSceneSection* Section = *It;
+	UMovieSceneTrack* Track = Section->GetTypedOuter<UMovieSceneTrack>();
+	if (Track)
 	{
-		UMovieSceneTrack* Track = Section->GetTypedOuter<UMovieSceneTrack>();
-		if (Track)
-		{
-			FScopedTransaction Transaction(LOCTEXT("SetSectionToKey", "Set Section To Key"));
-			Track->Modify();
-			Track->SetSectionToKey(bToggle ? nullptr : Section);
-		}
-
-		break;
+		FScopedTransaction Transaction(LOCTEXT("SetSectionToKey", "Set Section To Key"));
+		Track->Modify();
+		Track->SetSectionToKey(bToggle ? nullptr : Section);
 	}
 }
 
@@ -954,17 +951,14 @@ bool FSectionContextMenu::CanSetSectionToKey() const
 		return false;
 	}
 
-	for (UMovieSceneSection* Section : Sequencer->GetViewModel()->GetSelection()->GetSelectedSections())
+	auto It = Sequencer->GetViewModel()->GetSelection()->GetSelectedSections().CreateConstIterator();
+	UMovieSceneSection* Section = *It;
+	UMovieSceneTrack* Track = Section->GetTypedOuter<UMovieSceneTrack>();
+	if (Track && Section->GetBlendType().IsValid() && (Section->GetBlendType().Get() == EMovieSceneBlendType::Absolute 
+		|| Section->GetBlendType().Get() == EMovieSceneBlendType::Additive 
+		|| Section->GetBlendType().Get() == EMovieSceneBlendType::Override))
 	{
-		UMovieSceneTrack* Track = Section->GetTypedOuter<UMovieSceneTrack>();
-		if (Track && Section->GetBlendType().IsValid() && (Section->GetBlendType().Get() == EMovieSceneBlendType::Absolute 
-			|| Section->GetBlendType().Get() == EMovieSceneBlendType::Additive 
-			|| Section->GetBlendType().Get() == EMovieSceneBlendType::Override))
-		{
-			return true;
-		}
-
-		break;
+		return true;
 	}
 	return false;
 }
diff --git a/Engine/Source/Editor/Sequencer/Private/SequencerHotspots.cpp b/Engine/Source/Editor/Sequencer/Private/SequencerHotspots.cpp
--- a/Engine/Source/Editor/Sequencer/Private/SequencerHotspots.cpp
+++ b/Engine/Source/Editor/Sequencer/Private/SequencerHotspots.cpp
@@ -256,12 +256,12 @@ TOptional<FFrameNumber> FKeyHotspot::GetTime() const
 {
 	FFrameNumber Time = 0;
 
-	for (const FSequencerSelectedKey& Key : Keys)
+	if (auto It = Keys.CreateConstIterator(); It)
 	{
+		const FSequencerSelectedKey& Key = *It;
 		TArrayView<const FSequencerSelectedKey> FirstKey(&Key, 1);
 		TArrayView<FFrameNumber> FirstKeyTime(&Time, 1);
 		GetKeyTimes(FirstKey, FirstKeyTime);
-		break;
 	}
 
 	return Time;
diff --git a/Engine/Source/Editor/Sequencer/Private/Tools/SequencerEntityVisitor.cpp b/Engine/Source/Editor/Sequencer/Private/Tools/SequencerEntityVisitor.cpp
--- a/Engine/Source/Editor/Sequencer/Private/Tools/SequencerEntityVisitor.cpp
+++ b/Engine/Source/Editor/Sequencer/Private/Tools/SequencerEntityVisitor.cpp
@@ -165,10 +165,10 @@ void FSequencerEntityWalker::VisitAnyChannels(const ISequencerEntityVisitor& Vis
 	if (FTrackModel* TrackModel = InNode->CastThis<FTrackModel>())
 	{
 		// TODO: Top-level channel group -- there used to be only one at most but this assumption isn't necessarily true anymore
-		for (const TViewModelPtr<FChannelGroupModel>& TopLevelChannel : TrackModel->GetTopLevelChannels().IterateSubList<FChannelGroupModel>())
+		if (auto It = TrackModel->GetTopLevelChannels().IterateSubList<FChannelGroupModel>(); It)
 		{
+			const TViewModelPtr<FChannelGroupModel>& TopLevelChannel = *It;
 			ChannelGroup = TopLevelChannel.Get();
-			break;
 		}
 	}
 
diff --git a/Engine/Source/Editor/SequencerCore/Private/MVVM/Views/SOutlinerView.cpp b/Engine/Source/Editor/SequencerCore/Private/MVVM/Views/SOutlinerView.cpp
--- a/Engine/Source/Editor/SequencerCore/Private/MVVM/Views/SOutlinerView.cpp
+++ b/Engine/Source/Editor/SequencerCore/Private/MVVM/Views/SOutlinerView.cpp
@@ -792,8 +792,9 @@ TSharedPtr<SWidget> SOutlinerView::OnContextMenuOpening()
 	TSharedPtr<SWidget> ContextMenu;
 
 	// Open a context menu for the first selected item if it is selectable
-	for (TWeakViewModelPtr<IOutlinerExtension> WeakSelectedItem : GetSelectedItems())
+	if (auto It = GetSelectedItems().CreateConstIterator(); It)
 	{
+		TWeakViewModelPtr<IOutlinerExtension> WeakSelectedItem = *It;
 		TViewModelPtr<IOutlinerExtension> SelectedItem = WeakSelectedItem.Pin();
 
 		// Items are selectable by default, but can opt-in to conditionally selectable
@@ -806,7 +807,6 @@ TSharedPtr<SWidget> SOutlinerView::OnContextMenuOpening()
 			TSharedPtr<FEditorViewModel> EditorViewModel = GetOutlinerModel()->GetEditor();
 			ContextMenu = SelectedItem->CreateContextMenuWidget(EditorViewModel);
 		}
-		break;
 	}
 
 	// Otherwise, add a general menu for options
@@ -1144,8 +1144,9 @@ TSharedPtr<STrackLane> SOutlinerView::FindOrCreateParentLane(TViewModelPtr<IOutl
 		if (TrackAreaParameters.LaneType == ETrackAreaLaneType::Nested)
 		{
 			// Find the first outliner extension parent (should be the same as ParentTrackArea really)
-			for (const TViewModelPtr<IOutlinerExtension>& OutlinerExtItem : ParentTrackArea.AsModel()->GetAncestorsOfType<IOutlinerExtension>(true))
+			if (auto AncestorsIt = ParentTrackArea.AsModel()->GetAncestorsOfType<IOutlinerExtension>(true); AncestorsIt)
 			{
+				const TViewModelPtr<IOutlinerExtension>& OutlinerExtItem = *AncestorsIt;
 				TSharedPtr<STrackLane> TrackLane = TrackArea->FindTrackSlot(OutlinerExtItem);
 				if (!TrackLane)
 				{
diff --git a/Engine/Source/Editor/UMGEditor/Private/Designer/SDesignerView.cpp b/Engine/Source/Editor/UMGEditor/Private/Designer/SDesignerView.cpp
--- a/Engine/Source/Editor/UMGEditor/Private/Designer/SDesignerView.cpp
+++ b/Engine/Source/Editor/UMGEditor/Private/Designer/SDesignerView.cpp
@@ -1005,10 +1005,8 @@ FWidgetReference SDesignerView::GetSelectedWidget() const
 	// Only return a selected widget when we have only a single item selected.
 	if ( SelectedWidgets.Num() == 1 )
 	{
-		for ( TSet<FWidgetReference>::TConstIterator SetIt(SelectedWidgets); SetIt; ++SetIt )
-		{
-			return *SetIt;
-		}
+		TSet<FWidgetReference>::TConstIterator SetIt(SelectedWidgets);
+		return *SetIt;
 	}
 
 	return FWidgetReference();
diff --git a/Engine/Source/Editor/UMGEditor/Private/Extensions/CanvasSlotExtension.cpp b/Engine/Source/Editor/UMGEditor/Private/Extensions/CanvasSlotExtension.cpp
--- a/Engine/Source/Editor/UMGEditor/Private/Extensions/CanvasSlotExtension.cpp
+++ b/Engine/Source/Editor/UMGEditor/Private/Extensions/CanvasSlotExtension.cpp
@@ -373,8 +373,9 @@ FReply FCanvasSlotExtension::HandleAnchorDragging(const FGeometry& Geometry, con
 	{
 		float InverseSize = 1.0f / Designer->GetPreviewScale();
 
-		for ( FWidgetReference& Selection : SelectionCache )
+		if (SelectionCache.Num() >0)
 		{
+			FWidgetReference& Selection = SelectionCache[0];
 			UWidget* PreviewWidget = Selection.GetPreview();
 			if ( UCanvasPanel* Canvas = Cast<UCanvasPanel>(PreviewWidget->GetParent()) )
 			{
diff --git a/Engine/Source/Editor/UMGEditor/Private/Navigation/SWidgetDesignerNavigation.cpp b/Engine/Source/Editor/UMGEditor/Private/Navigation/SWidgetDesignerNavigation.cpp
--- a/Engine/Source/Editor/UMGEditor/Private/Navigation/SWidgetDesignerNavigation.cpp
+++ b/Engine/Source/Editor/UMGEditor/Private/Navigation/SWidgetDesignerNavigation.cpp
@@ -429,11 +429,9 @@ void SWidgetDesignerNavigation::HandleEditorSelectionChanged()
 		const TSet<FWidgetReference>& SelectedWidgetReferences = BlueprintEditorPinned->GetSelectedWidgets();
 		if (SelectedWidgetReferences.Num() > 0)
 		{
-			for (const FWidgetReference& WidgetReference : SelectedWidgetReferences)
-			{
-				NavigationEventSimulationView->SelectWidget(WidgetReference.GetPreviewSlate());
-				break;
-			}
+			auto It = SelectedWidgetReferences.CreateConstIterator();
+			const FWidgetReference& WidgetReference = *It;
+			NavigationEventSimulationView->SelectWidget(WidgetReference.GetPreviewSlate());
 		}
 		else
 		{
diff --git a/Engine/Source/Editor/UMGEditor/Private/WidgetBlueprintEditor.cpp b/Engine/Source/Editor/UMGEditor/Private/WidgetBlueprintEditor.cpp
--- a/Engine/Source/Editor/UMGEditor/Private/WidgetBlueprintEditor.cpp
+++ b/Engine/Source/Editor/UMGEditor/Private/WidgetBlueprintEditor.cpp
@@ -2035,11 +2035,11 @@ void FWidgetBlueprintEditor::OnBuildCustomContextMenuForGuid(FMenuBuilder& MenuB
 				UWidget* BoundWidget = nullptr;
 				bool bNotBound = true;
 				TSharedPtr<ISequencer>& ActiveSequencer = GetSequencer();
-				for (TWeakObjectPtr<> WeakObjectPtr : ActiveSequencer->FindObjectsInCurrentSequence(ObjectBinding))
+				if (auto ObjectsView = ActiveSequencer->FindObjectsInCurrentSequence(ObjectBinding); ObjectsView.Num() > 0)
 				{
+					TWeakObjectPtr<> WeakObjectPtr = ObjectsView[0];
 					BoundWidget = Cast<UWidget>(WeakObjectPtr.Get());
 					bNotBound = false;
-					break;
 				}
 
 				if (bNotBound || (BoundWidget && SelectedWidget.GetPreview()->GetTypedOuter<UWidgetTree>() == BoundWidget->GetTypedOuter<UWidgetTree>()))
diff --git a/Engine/Source/Editor/UnrealEd/Private/Commandlets/DumpLightFunctionMaterialInfo.cpp b/Engine/Source/Editor/UnrealEd/Private/Commandlets/DumpLightFunctionMaterialInfo.cpp
--- a/Engine/Source/Editor/UnrealEd/Private/Commandlets/DumpLightFunctionMaterialInfo.cpp
+++ b/Engine/Source/Editor/UnrealEd/Private/Commandlets/DumpLightFunctionMaterialInfo.cpp
@@ -58,8 +58,10 @@ int32 UDumpLightFunctionMaterialInfoCommandlet::Main(const FString& Params)
 	TSet<UMaterialInterface*> MaterialsToCompile;
 	TSet<UMaterialInterface*> MaterialsToAnalyse;
 
-	for (ITargetPlatform* Platform : Platforms)
+	// only run for a single platform as this is enough to know if a light function material will be compatible with that LFAtlas.
+	if (Platforms.Num() > 0)
 	{
+		ITargetPlatform* Platform = Platforms[0];
 		UE_LOG(LogDumpLightFunctionMaterialInfo, Display, TEXT("Compiling shaders for %s..."), *Platform->PlatformName());
 		for (const FAssetData& AssetData : MaterialAssets)
 		{
@@ -214,10 +216,6 @@ int32 UDumpLightFunctionMaterialInfoCommandlet::Main(const FString& Params)
 				}
 			}
 		}
-
-
-
-		break;	// only run for a single platform as this is enough to know if a light function material will be compatible with that LFAtlas.
 	} // Platforms
 	
 
diff --git a/Engine/Source/Editor/UnrealEd/Private/Commandlets/SubstrateCommandlet.cpp b/Engine/Source/Editor/UnrealEd/Private/Commandlets/SubstrateCommandlet.cpp
--- a/Engine/Source/Editor/UnrealEd/Private/Commandlets/SubstrateCommandlet.cpp
+++ b/Engine/Source/Editor/UnrealEd/Private/Commandlets/SubstrateCommandlet.cpp
@@ -66,8 +66,10 @@ int32 USubstrateCommandlet::Main(const FString& Params)
 	TSet<UMaterialInterface*> MaterialsToCompile;
 #endif
 
-	for (ITargetPlatform* Platform : Platforms)
+	// only run for a single platform as this is enough to know if material is substrate or not
+	if (Platforms.Num() > 0)
 	{
+		ITargetPlatform* Platform = Platforms[0];
 		UE_LOG(LogSubstrateCommandlet, Display, TEXT("Compiling shaders for %s..."), *Platform->PlatformName());
 		for (const FAssetData& AssetData : MaterialAssets)
 		{
@@ -221,9 +223,6 @@ int32 USubstrateCommandlet::Main(const FString& Params)
 			}
 		}
 #endif
-
-
-		break;	// only run for a single platform as this is enough to know if material is substrate or not
 	} // Platforms
 	
 
diff --git a/Engine/Source/Editor/UnrealEd/Private/Cooker/CookPackageData.cpp b/Engine/Source/Editor/UnrealEd/Private/Cooker/CookPackageData.cpp
--- a/Engine/Source/Editor/UnrealEd/Private/Cooker/CookPackageData.cpp
+++ b/Engine/Source/Editor/UnrealEd/Private/Cooker/CookPackageData.cpp
@@ -2996,13 +2996,13 @@ uint32 FRequestQueue::ReadyRequestsNum() const
 
 FPackageData* FRequestQueue::PopReadyRequest()
 {
-	for (auto Iterator = UrgentRequests.CreateIterator(); Iterator; ++Iterator)
+	if (auto Iterator = UrgentRequests.CreateIterator(); Iterator)
 	{
 		FPackageData* PackageData = *Iterator;
 		Iterator.RemoveCurrent();
 		return PackageData;
 	}
-	for (auto Iterator = NormalRequests.CreateIterator(); Iterator; ++Iterator)
+	if (auto Iterator = NormalRequests.CreateIterator(); Iterator)
 	{
 		FPackageData* PackageData = *Iterator;
 		Iterator.RemoveCurrent();
diff --git a/Engine/Source/Editor/UnrealEd/Private/EditorEngine.cpp b/Engine/Source/Editor/UnrealEd/Private/EditorEngine.cpp
--- a/Engine/Source/Editor/UnrealEd/Private/EditorEngine.cpp
+++ b/Engine/Source/Editor/UnrealEd/Private/EditorEngine.cpp
@@ -4698,14 +4698,14 @@ void UEditorEngine::OnPostSaveWorld(UWorld* World, FObjectPostSaveContext Object
 APlayerStart* UEditorEngine::CheckForPlayerStart()
 {
 	UWorld* IteratorWorld = GWorld;
-	for( TActorIterator<APlayerStart> It(IteratorWorld); It; ++It )
+	if ( TActorIterator<APlayerStart> It(IteratorWorld); It)
 	{
 		// Return the found start.
 		return *It;
 	}
 
 	// No player start was found, return NULL.
-	return NULL;
+	return nullptr;
 }
 
 void UEditorEngine::CloseEntryPopupWindow()
diff --git a/Engine/Source/Editor/UnrealEd/Private/Kismet2/DebuggerCommands.cpp b/Engine/Source/Editor/UnrealEd/Private/Kismet2/DebuggerCommands.cpp
--- a/Engine/Source/Editor/UnrealEd/Private/Kismet2/DebuggerCommands.cpp
+++ b/Engine/Source/Editor/UnrealEd/Private/Kismet2/DebuggerCommands.cpp
@@ -2060,7 +2060,7 @@ bool FInternalPlayWorldCommandCallbacks::CanPossessEjectPlayer()
 {
 	if ((IsInSIE() || IsInPIE()) && !IsStoppedAtBreakpoint())
 	{
-		for (auto It = GUnrealEd->SlatePlayInEditorMap.CreateIterator(); It; ++It)
+		if (auto It = GUnrealEd->SlatePlayInEditorMap.CreateIterator(); It)
 		{
 			return It.Value().DestinationSlateViewport.IsValid();
 		}
diff --git a/Engine/Source/Runtime/Engine/Private/DemoNetDriver.cpp b/Engine/Source/Runtime/Engine/Private/DemoNetDriver.cpp
--- a/Engine/Source/Runtime/Engine/Private/DemoNetDriver.cpp
+++ b/Engine/Source/Runtime/Engine/Private/DemoNetDriver.cpp
@@ -4952,10 +4952,9 @@ void UDemoNetConnection::HandleClientPlayer(APlayerController* PC, UNetConnectio
 	DemoDriver->SetSpectatorController(PC);
 
 	// Find a player start, if one exists
-	for (TActorIterator<APlayerStart> PlayerStartIt(Driver->World); PlayerStartIt; ++PlayerStartIt)
+	if (TActorIterator<APlayerStart> PlayerStartIt(Driver->World); PlayerStartIt)
 	{
 		PC->SetInitialLocationAndRotation(PlayerStartIt->GetActorLocation(), PlayerStartIt->GetActorRotation());
-		break;
 	}
 }
 
diff --git a/Engine/Source/Runtime/Engine/Private/GameplayStatics.cpp b/Engine/Source/Runtime/Engine/Private/GameplayStatics.cpp
--- a/Engine/Source/Runtime/Engine/Private/GameplayStatics.cpp
+++ b/Engine/Source/Runtime/Engine/Private/GameplayStatics.cpp
@@ -1082,7 +1082,7 @@ AActor* UGameplayStatics::GetActorOfClass(const UObject* WorldContextObject, TSu
 	{
 		if (UWorld* World = GEngine->GetWorldFromContextObject(WorldContextObject, EGetWorldErrorMode::LogAndReturnNull))
 		{
-			for (TActorIterator<AActor> It(World, ActorClass); It; ++It)
+			if (TActorIterator<AActor> It(World, ActorClass); It)
 			{
 				AActor* Actor = *It;
 				return Actor;
diff --git a/Engine/Source/Runtime/Engine/Private/Materials/Material.cpp b/Engine/Source/Runtime/Engine/Private/Materials/Material.cpp
--- a/Engine/Source/Runtime/Engine/Private/Materials/Material.cpp
+++ b/Engine/Source/Runtime/Engine/Private/Materials/Material.cpp
@@ -5743,10 +5743,11 @@ void UMaterial::RebuildShadingModelField()
 				// Volumetric shading model do support Unlit as an EmissiveOnly mode, on top of the lit mode.
 				TArray<UMaterialExpressionSubstrateVolumetricFogCloudBSDF*> ShadingModelExpressions;
 				GetAllExpressionsInMaterialAndFunctionsOfType(ShadingModelExpressions);
-				for (UMaterialExpressionSubstrateVolumetricFogCloudBSDF* MatExpr : ShadingModelExpressions)
+				if (ShadingModelExpressions.Num() > 0)
 				{
+					// We can only keep a single shading model due to an assert in UMaterialInstance::UpdateOverridableBaseProperties which want a single shading model for non material attribute workflow.
+					UMaterialExpressionSubstrateVolumetricFogCloudBSDF* MatExpr = ShadingModelExpressions[0];
 					ShadingModel = MatExpr->bEmissiveOnly ? MSM_Unlit : MSM_DefaultLit;
-					break; // We can only keep a single shading model due to an assert in UMaterialInstance::UpdateOverridableBaseProperties which want a single shading model for non material attribute workflow.
 				}
 				BlendMode = EBlendMode::BLEND_Additive;
 			}
diff --git a/Engine/Source/Runtime/Engine/Private/NetDriver.cpp b/Engine/Source/Runtime/Engine/Private/NetDriver.cpp
--- a/Engine/Source/Runtime/Engine/Private/NetDriver.cpp
+++ b/Engine/Source/Runtime/Engine/Private/NetDriver.cpp
@@ -6429,10 +6429,9 @@ void UNetDriver::DrawNetDriverDebug()
 	}
 
 	ULocalPlayer*	LocalPlayer = NULL;
-	for(FLocalPlayerIterator It(GEngine, LocalWorld);It;++It)
+	if(FLocalPlayerIterator It(GEngine, LocalWorld);It)
 	{
 		LocalPlayer = *It;
-		break;
 	}
 	if (!LocalPlayer)
 	{
diff --git a/Engine/Source/Runtime/Engine/Private/PlayerController.cpp b/Engine/Source/Runtime/Engine/Private/PlayerController.cpp
--- a/Engine/Source/Runtime/Engine/Private/PlayerController.cpp
+++ b/Engine/Source/Runtime/Engine/Private/PlayerController.cpp
@@ -1965,14 +1965,13 @@ void APlayerController::NotifyLoadedWorld(FName WorldPackageName, bool bFinalDes
 	// place the camera at the first playerstart we can find
 	SetViewTarget(this);
 	
-	for( TActorIterator<APlayerStart> It(GetWorld()); It; ++It )
+	if (TActorIterator<APlayerStart> It(GetWorld()); It)
 	{
 		APlayerStart* P = *It;
 
 		FRotator SpawnRotation(ForceInit);
 		SpawnRotation.Yaw = P->GetActorRotation().Yaw;
 		SetInitialLocationAndRotation(P->GetActorLocation(), SpawnRotation);
-		break;
 	}
 }
 
diff --git a/Engine/Source/Runtime/Engine/Private/ReplayHelper.cpp b/Engine/Source/Runtime/Engine/Private/ReplayHelper.cpp
--- a/Engine/Source/Runtime/Engine/Private/ReplayHelper.cpp
+++ b/Engine/Source/Runtime/Engine/Private/ReplayHelper.cpp
@@ -389,10 +389,9 @@ APlayerController* FReplayHelper::CreateSpectatorController(UNetConnection* Conn
 		NewDemoController->PlayerState->SetIsOnlyASpectator(true);
 	}
 
-	for (TActorIterator<APlayerStart> It(World.Get()); It; ++It)
+	if (TActorIterator<APlayerStart> It(World.Get()); It)
 	{
 		NewDemoController->SetInitialLocationAndRotation(It->GetActorLocation(), It->GetActorRotation());
-		break;
 	}
 
 	NewDemoController->SetReplicates(true);
diff --git a/Engine/Source/Runtime/Experimental/Animation/Constraints/Private/Transform/TransformableHandle.cpp b/Engine/Source/Runtime/Experimental/Animation/Constraints/Private/Transform/TransformableHandle.cpp
--- a/Engine/Source/Runtime/Experimental/Animation/Constraints/Private/Transform/TransformableHandle.cpp
+++ b/Engine/Source/Runtime/Experimental/Animation/Constraints/Private/Transform/TransformableHandle.cpp
@@ -554,8 +554,10 @@ void UTransformableComponentHandle::ResolveBoundObjects(FMovieSceneSequenceID Lo
 		}
 	};
 	
-	for (const TWeakObjectPtr<>& ParentObject : ConstraintBindingID.ResolveBoundObjects(LocalSequenceID, SharedPlaybackState))
+	//just do one
+	if (auto BoundObjectsView = ConstraintBindingID.ResolveBoundObjects(LocalSequenceID, SharedPlaybackState); BoundObjectsView.Num() > 0)
 	{
+		const TWeakObjectPtr<>& ParentObject = BoundObjectsView[0];
 		if (const AActor* Actor = Cast<AActor>(ParentObject.Get()))
 		{
 			const TInlineComponentArray<USceneComponent*> Components(Actor);
@@ -570,7 +572,6 @@ void UTransformableComponentHandle::ResolveBoundObjects(FMovieSceneSequenceID Lo
 		{
 			ReplaceComponent(Comp);
 		}
-		break; //just do one
 	}
 }
 
diff --git a/Engine/Source/Runtime/Experimental/Chaos/Private/GeometryCollection/GeometryCollectionConvexUtility.cpp b/Engine/Source/Runtime/Experimental/Chaos/Private/GeometryCollection/GeometryCollectionConvexUtility.cpp
--- a/Engine/Source/Runtime/Experimental/Chaos/Private/GeometryCollection/GeometryCollectionConvexUtility.cpp
+++ b/Engine/Source/Runtime/Experimental/Chaos/Private/GeometryCollection/GeometryCollectionConvexUtility.cpp
@@ -481,8 +481,9 @@ void CreateNonoverlappingConvexHulls(
 	auto OnlyConvex = [&TransformToConvexIndices](int32 Bone) -> int32
 	{
 		ensure(TransformToConvexIndices[Bone].Num() <= 1);
-		for (int32 ConvexIndex : TransformToConvexIndices[Bone])
+		if (auto It = TransformToConvexIndices[Bone].CreateConstIterator(); It)
 		{
+			const int32 ConvexIndex = *It;
 			return ConvexIndex;
 		}
 		return INDEX_NONE;
@@ -2867,7 +2868,6 @@ void FGeometryCollectionConvexUtility::ConvertImplicitToConvexArray(const Chaos:
 		}
 		break;
 	}
-	break;
 	case Chaos::ImplicitObjectType::Convex:
 	{
 		if (const Chaos::FConvex* Convex = InImplicit.template GetObject<Chaos::FConvex>())
diff --git a/Engine/Source/Runtime/Renderer/Private/RendererScene.cpp b/Engine/Source/Runtime/Renderer/Private/RendererScene.cpp
--- a/Engine/Source/Runtime/Renderer/Private/RendererScene.cpp
+++ b/Engine/Source/Runtime/Renderer/Private/RendererScene.cpp
@@ -504,16 +504,8 @@ bool FPixelInspectorData::AddPixelInspectorRequest(FPixelInspectorRequest *Pixel
 	//Remove the oldest request since the new request use the buffer
 	if (Requests.Num() > 1)
 	{
-		FVector2f FirstKey(-1, -1);
-		for (auto kvp : Requests)
-		{
-			FirstKey = kvp.Key;
-			break;
-		}
-		if (Requests.Contains(FirstKey))
-		{
-			Requests.Remove(FirstKey);
-		}
+		auto It = Requests.CreateIterator();
+		It.RemoveCurrent();
 	}
 	Requests.Add(ViewportUV, PixelInspectorRequest);
 	return true;
